World Bookmarks initial implementation
* WorldBookmark are a new type of assets that can capture the state of various subsystems in the editor and restore it upon request.
* The states that are currently captured in bookmarks are:
    -> The currently loaded world
    -> Loaded/unloaded data layers
    -> Level editor camera properties
    -> The active actor editor context
    -> Loaded world partition regions/volumes
* States types can be registered in the EditorStateSubsystem, so any system could write it's own & register it to have it included in the bookmarks
* Added WorldBookmark browser that list the bookmarks relevant to the current world.

diff --git a/Engine/Source/Developer/AssetTools/Private/AssetTools.cpp b/Engine/Source/Developer/AssetTools/Private/AssetTools.cpp
--- a/Engine/Source/Developer/AssetTools/Private/AssetTools.cpp
+++ b/Engine/Source/Developer/AssetTools/Private/AssetTools.cpp
@@ -155,7 +155,8 @@ public:
 	virtual bool CanLocalize() const override { return AssetDefinitionPtr.Get()->CanLocalize(FAssetData()).IsSupported(); }
 	virtual bool IsImportedAsset() const override { return AssetDefinitionPtr.Get()->CanImport(); }
 	virtual bool CanMerge() const override { return AssetDefinitionPtr.Get()->CanMerge(); }
-	
+	virtual bool ShouldFindEditorForAsset() const override { return AssetDefinitionPtr.Get()->ShouldFindEditorForAsset(); }
+
 	virtual bool CanRename(const FAssetData& InAsset, FText* OutErrorMsg) const override
     {
 		FAssetSupportResponse Response = AssetDefinitionPtr.Get()->CanRename(InAsset);
diff --git a/Engine/Source/Developer/AssetTools/Public/AssetTypeActions_Base.h b/Engine/Source/Developer/AssetTools/Public/AssetTypeActions_Base.h
--- a/Engine/Source/Developer/AssetTools/Public/AssetTypeActions_Base.h
+++ b/Engine/Source/Developer/AssetTools/Public/AssetTypeActions_Base.h
@@ -31,6 +31,11 @@ public:
 		return Object->GetName();
 	}
 
+	virtual bool ShouldFindEditorForAsset() const override
+	{
+		return true;
+	}
+
 	virtual void OpenAssetEditor( const TArray<UObject*>& InObjects, TSharedPtr<class IToolkitHost> EditWithinLevelEditor = TSharedPtr<IToolkitHost>() ) override
 	{
 		FSimpleAssetEditor::CreateEditor(EToolkitMode::Standalone, EditWithinLevelEditor, InObjects);
diff --git a/Engine/Source/Developer/AssetTools/Public/IAssetTypeActions.h b/Engine/Source/Developer/AssetTools/Public/IAssetTypeActions.h
--- a/Engine/Source/Developer/AssetTools/Public/IAssetTypeActions.h
+++ b/Engine/Source/Developer/AssetTools/Public/IAssetTypeActions.h
@@ -56,6 +56,9 @@ public:
 
 	/** Generates a menu section for the specified objects. */
 	virtual void GetActions(const TArray<UObject*>& InObjects, struct FToolMenuSection& Section) { }
+	
+	/** Returns true if an asset editor window should be created for this asset. */
+	virtual bool ShouldFindEditorForAsset() const = 0;
 
 	/** Opens the asset editor for the specified objects. If EditWithinLevelEditor is valid, the world-centric editor will be used. */
 	virtual void OpenAssetEditor( const TArray<UObject*>& InObjects, TSharedPtr<IToolkitHost> EditWithinLevelEditor = TSharedPtr<IToolkitHost>() ) = 0;
diff --git a/Engine/Source/Editor/AssetDefinition/Public/AssetDefinition.h b/Engine/Source/Editor/AssetDefinition/Public/AssetDefinition.h
--- a/Engine/Source/Editor/AssetDefinition/Public/AssetDefinition.h
+++ b/Engine/Source/Editor/AssetDefinition/Public/AssetDefinition.h
@@ -487,6 +487,11 @@ public:
 		return TArray<FAssetData>(ActivateArgs.Assets);
 	}
 	
+	virtual bool ShouldFindEditorForAsset() const
+	{ 
+		return true;
+	}
+
 	/** Get open support for the method.  Includes required information before we call OpenAsset. */
 	virtual FAssetOpenSupport GetAssetOpenSupport(const FAssetOpenSupportArgs& OpenSupportArgs) const
 	{
@@ -516,7 +521,7 @@ public:
 	{
 		return FAssetSupportResponse::Supported();
 	}
-
+	
 	
 	// Importing
 	virtual bool CanImport() const { return false; }
diff --git a/Engine/Source/Editor/ContentBrowser/Private/FrontendFilters.cpp b/Engine/Source/Editor/ContentBrowser/Private/FrontendFilters.cpp
--- a/Engine/Source/Editor/ContentBrowser/Private/FrontendFilters.cpp
+++ b/Engine/Source/Editor/ContentBrowser/Private/FrontendFilters.cpp
@@ -7,6 +7,7 @@
 #include "AssetRegistry/AssetData.h"
 #include "AssetRegistry/AssetRegistryModule.h"
 #include "AssetRegistry/IAssetRegistry.h"
+#include "AssetRegistry/AssetRegistryHelpers.h"
 #include "AssetToolsModule.h"
 #include "Blueprint/BlueprintSupport.h"
 #include "CollectionManagerModule.h"
@@ -104,6 +105,22 @@ namespace FrontendFilterHelper
 			}
 		}
 	}
+
+	static const FName NAME_BookmarkWorld(TEXT("WorldName"));
+	bool GetWorldBookmark(const IAssetRegistry* AssetRegistry, const FAssetData& InAssetData, FAssetData& OutValue)
+	{
+		// This would soon to be replace by World Asset concept.
+		FString WorldAssetName;
+		if (InAssetData.GetTagValue(NAME_BookmarkWorld, WorldAssetName))
+		{
+			FSoftObjectPath WorldAsset = FSoftObjectPath(WorldAssetName);
+			UAssetRegistryHelpers::FixupRedirectedAssetPath(WorldAsset);
+			OutValue = AssetRegistry->GetAssetByObjectPath(WorldAsset);
+			return true;
+		}
+
+		return false;
+	}
 }
 
 /////////////////////////////////////////
@@ -1297,6 +1314,15 @@ bool FFrontendFilter_InUseByLoadedLevels::PassesFilter(FAssetFilterType InItem)
 	FAssetData ItemAssetData;
 	if (InItem.Legacy_TryGetAssetData(ItemAssetData))
 	{
+		FString WorldAssetName;
+		if (GWorld && ItemAssetData.GetTagValue(FrontendFilterHelper::NAME_BookmarkWorld, WorldAssetName))
+		{
+			FSoftObjectPath WorldAsset = FSoftObjectPath(WorldAssetName);
+			UAssetRegistryHelpers::FixupRedirectedAssetPath(WorldAsset);
+			FSoftObjectPath CurrentWorld = FSoftObjectPath(GWorld);
+			return (CurrentWorld == WorldAsset);
+		}
+
 		if (UObject* Asset = ItemAssetData.FastGetAsset(false))
 		{
 			const bool bUnreferenced = !Asset->HasAnyMarks(OBJECTMARK_TagExp);
@@ -1360,6 +1386,12 @@ bool FFrontendFilter_UsedInAnyLevel::PassesFilter(FAssetFilterType InItem) const
 	FAssetData ItemAssetData;
 	if (InItem.Legacy_TryGetAssetData(ItemAssetData))
 	{
+		FAssetData WorldInBookmark;
+		if (FrontendFilterHelper::GetWorldBookmark(AssetRegistry, ItemAssetData, WorldInBookmark))
+		{
+			return WorldInBookmark.IsValid();
+		}
+
 		return LevelsDependencies.Contains(ItemAssetData.PackageName);
 	}
 	return false;
@@ -1398,6 +1430,12 @@ bool FFrontendFilter_NotUsedInAnyLevel::PassesFilter(FAssetFilterType InItem) co
 	FAssetData ItemAssetData;
 	if (InItem.Legacy_TryGetAssetData(ItemAssetData))
 	{
+		FAssetData WorldInBookmark;
+		if (FrontendFilterHelper::GetWorldBookmark(AssetRegistry, ItemAssetData, WorldInBookmark))
+		{
+			return !WorldInBookmark.IsValid();
+		}
+
 		return !LevelsDependencies.Contains(ItemAssetData.PackageName);
 	}
 	return false;
diff --git a/Engine/Source/Editor/DataLayerEditor/Private/DataLayer/DataLayerEditorState.cpp b/Engine/Source/Editor/DataLayerEditor/Private/DataLayer/DataLayerEditorState.cpp
new file mode 100644
--- /dev/null
+++ b/Engine/Source/Editor/DataLayerEditor/Private/DataLayer/DataLayerEditorState.cpp
@@ -0,0 +1,119 @@
+// Copyright Epic Games, Inc. All Rights Reserved.
+
+#include "DataLayer/DataLayerEditorState.h"
+#include "DataLayer/DataLayerEditorSubsystem.h"
+#include "WorldPartition/DataLayer/DataLayerInstance.h"
+#include "WorldPartition/DataLayer/DataLayerManager.h"
+#include "Algo/Transform.h"
+
+#include UE_INLINE_GENERATED_CPP_BY_NAME(DataLayerEditorState)
+
+#define LOCTEXT_NAMESPACE "DataLayersEditorState"
+
+UDataLayerEditorState::UDataLayerEditorState(const FObjectInitializer& ObjectInitializer)
+	: Super(ObjectInitializer)
+{
+}
+
+FText UDataLayerEditorState::GetCategoryText() const
+{
+	return FText(LOCTEXT("DataLayersEditorStateCategoryText", "Data Layers"));
+}
+
+bool UDataLayerEditorState::CaptureState()
+{
+	if (UDataLayerManager* DataLayerManager = UDataLayerManager::GetDataLayerManager(GetStateWorld()))
+	{
+		DataLayerManager->ForEachDataLayerInstance([this](const UDataLayerInstance* DataLayerInstance)
+		{
+			if (DataLayerInstance && DataLayerInstance->GetAsset())
+			{
+				if (DataLayerInstance->IsLoadedInEditor() != DataLayerInstance->IsInitiallyLoadedInEditor())
+				{
+					if (DataLayerInstance->IsLoadedInEditor())
+					{
+						LoadedDataLayers.Add(DataLayerInstance->GetAsset());
+					}
+					else
+					{
+						NotLoadedDataLayers.Add(DataLayerInstance->GetAsset());
+					}
+				}
+			}
+
+			return true;
+		});
+
+		return true;
+	}
+
+	return false;
+}
+
+bool UDataLayerEditorState::RestoreState() const
+{
+	if (UDataLayerManager* DataLayerManager = UDataLayerManager::GetDataLayerManager(GetStateWorld()))
+	{
+		// Gather overriden state for all DataLayer instance
+		TMap<const UDataLayerInstance*, bool> DataLayersLoadedInEditor;
+
+		// Unloaded
+		for (const TObjectPtr<const UDataLayerAsset>& DataLayerAsset : NotLoadedDataLayers)
+		{
+			if (DataLayerAsset != nullptr)
+			{
+				const UDataLayerInstance* DataLayerInstance = DataLayerManager->GetDataLayerInstance(DataLayerAsset);
+				if (DataLayerInstance != nullptr)
+				{
+					DataLayersLoadedInEditor.Emplace(DataLayerInstance, false);
+				}
+			}
+		}
+
+		// Loaded
+		for (const TObjectPtr<const UDataLayerAsset>& DataLayerAsset : LoadedDataLayers)
+		{
+			if (DataLayerAsset != nullptr)
+			{
+				const UDataLayerInstance* DataLayerInstance = DataLayerManager->GetDataLayerInstance(DataLayerAsset);
+				if (DataLayerInstance != nullptr)
+				{
+					DataLayersLoadedInEditor.Emplace(DataLayerInstance, true);
+				}
+			}
+		}
+
+		TArray<UDataLayerInstance*> DataLayersLoaded;
+		TArray<UDataLayerInstance*> DataLayersUnloaded;
+		DataLayerManager->ForEachDataLayerInstance([&DataLayersLoaded, &DataLayersUnloaded, &DataLayersLoadedInEditor](UDataLayerInstance* DataLayerInstance)
+		{
+			bool bDataLayerLoadedInEditor = DataLayerInstance->IsInitiallyLoadedInEditor();
+			if (bool* bBookmarkValue = DataLayersLoadedInEditor.Find(DataLayerInstance))
+			{
+				bDataLayerLoadedInEditor = *bBookmarkValue;
+			}
+
+			TArray<UDataLayerInstance*>& DataLayersArray = bDataLayerLoadedInEditor ? DataLayersLoaded : DataLayersUnloaded;
+			DataLayersArray.Add(DataLayerInstance);
+
+			// Visibility of Data Layers is not currently saved to users settings or world bookmarks
+			// Until then, we restore them to their initial visibility
+			DataLayerInstance->SetVisible(DataLayerInstance->IsInitiallyVisible());
+
+			return true;
+		});
+
+		if (UDataLayerEditorSubsystem* DataLayerEditorSubsystem = UDataLayerEditorSubsystem::Get())
+		{
+			const bool LOAD_DATALAYER = true;
+			const bool UNLOAD_DATALAYER = false;
+			const bool FROM_USER_CHANGE = true;
+			DataLayerEditorSubsystem->SetDataLayersIsLoadedInEditor(DataLayersLoaded, LOAD_DATALAYER, FROM_USER_CHANGE);
+			DataLayerEditorSubsystem->SetDataLayersIsLoadedInEditor(DataLayersUnloaded, UNLOAD_DATALAYER, FROM_USER_CHANGE);
+		}
+	}
+
+	return true;
+}
+
+#undef LOCTEXT_NAMESPACE
diff --git a/Engine/Source/Editor/DataLayerEditor/Private/DataLayer/DataLayerEditorSubsystem.cpp b/Engine/Source/Editor/DataLayerEditor/Private/DataLayer/DataLayerEditorSubsystem.cpp
--- a/Engine/Source/Editor/DataLayerEditor/Private/DataLayer/DataLayerEditorSubsystem.cpp
+++ b/Engine/Source/Editor/DataLayerEditor/Private/DataLayer/DataLayerEditorSubsystem.cpp
@@ -2,12 +2,14 @@
 
 #include "DataLayer/DataLayerEditorSubsystem.h"
 #include "DataLayer/DataLayerAction.h"
+#include "DataLayer/DataLayerEditorState.h"
 #include "Containers/EnumAsByte.h"
 #include "Containers/StringConv.h"
 #include "Containers/UnrealString.h"
 #include "CoreGlobals.h"
 #include "Editor.h"
 #include "Editor/EditorEngine.h"
+#include "EditorState/EditorStateSubsystem.h"
 #include "Engine/Engine.h"
 #include "Engine/EngineTypes.h"
 #include "Engine/Level.h"
@@ -363,6 +365,9 @@ void UDataLayerEditorSubsystem::Initialize(FSubsystemCollectionBase& Collection)
 
 	// Register actor descriptor loading filter
 	IWorldPartitionActorLoaderInterface::RegisterActorDescFilter(MakeShareable<IWorldPartitionActorLoaderInterface::FActorDescFilter>(new FDataLayerActorDescFilter(this)));
+
+	Collection.InitializeDependency<UEditorStateSubsystem>();
+	UEditorStateSubsystem::RegisterEditorStateType<UDataLayerEditorState>();
 }
 
 void UDataLayerEditorSubsystem::Deinitialize()
@@ -381,6 +386,8 @@ void UDataLayerEditorSubsystem::Deinitialize()
 	FActorPrimitiveColorHandler::Get().UnregisterPrimitiveColorHandler(NAME_CurrentDataLayerColor);
 	FActorPrimitiveColorHandler::Get().UnregisterPrimitiveColorHandler(NAME_ExternalDataLayerColor);
 #endif
+
+	UEditorStateSubsystem::UnregisterEditorStateType<UDataLayerEditorState>();
 }
 
 const UExternalDataLayerInstance* UDataLayerEditorSubsystem::GetActorSpawningExternalDataLayerInstance(AActor* InActor) const
diff --git a/Engine/Source/Editor/DataLayerEditor/Public/DataLayer/DataLayerEditorState.h b/Engine/Source/Editor/DataLayerEditor/Public/DataLayer/DataLayerEditorState.h
new file mode 100644
--- /dev/null
+++ b/Engine/Source/Editor/DataLayerEditor/Public/DataLayer/DataLayerEditorState.h
@@ -0,0 +1,28 @@
+// Copyright Epic Games, Inc. All Rights Reserved.
+
+#pragma once
+
+#include "EditorState/WorldEditorState.h"
+#include "DataLayerEditorState.generated.h"
+
+class UDataLayerAsset;
+
+UCLASS()
+class UDataLayerEditorState : public UWorldDependantEditorState
+{
+	GENERATED_UCLASS_BODY()
+
+	//~ Begin UEditorState Interface
+	virtual FText GetCategoryText() const override;
+private:
+	virtual bool CaptureState() override;
+	virtual bool RestoreState() const override;
+	//~ End UEditorState Interface
+
+private:
+	UPROPERTY(EditAnywhere, Category = DataLayers)
+	TArray<TObjectPtr<const UDataLayerAsset>> NotLoadedDataLayers;
+
+	UPROPERTY(EditAnywhere, Category = DataLayers)
+	TArray<TObjectPtr<const UDataLayerAsset>> LoadedDataLayers;
+};
diff --git a/Engine/Source/Editor/EditorStyle/Private/StarshipStyle.cpp b/Engine/Source/Editor/EditorStyle/Private/StarshipStyle.cpp
--- a/Engine/Source/Editor/EditorStyle/Private/StarshipStyle.cpp
+++ b/Engine/Source/Editor/EditorStyle/Private/StarshipStyle.cpp
@@ -2062,6 +2062,15 @@ void FStarshipEditorStyle::FStyle::SetupWorldPartitionStyles()
 	Set("WorldPartition.ShowCoords", new IMAGE_BRUSH_SVG("WorldPartition/Coordinate_16", Icon16x16));
 	Set("WorldPartition.ShowMouseCoords", new IMAGE_BRUSH_SVG("WorldPartition/MouseCoordinates_16", Icon16x16));
 
+	Set("WorldPartition.WorldBookmark", new IMAGE_BRUSH_SVG("Starship/AssetIcons/WorldBookmark_64", Icon64x64));
+	Set("WorldPartition.WorldBookmark.TabIcon", new IMAGE_BRUSH_SVG("Starship/AssetIcons/WorldBookmark_16", Icon16x16));
+	Set("WorldPartition.WorldBookmark.GoTo", new IMAGE_BRUSH_SVG("Starship/Common/NextArrow", Icon16x16));
+	Set("WorldPartition.WorldBookmark.Update", new CORE_IMAGE_BRUSH_SVG("Starship/Common/Update", Icon20x20));
+	Set("WorldPartition.WorldBookmark.Create", new CORE_IMAGE_BRUSH_SVG("Starship/Common/plus", Icon16x16, FStyleColors::AccentGreen));
+	Set("WorldPartition.WorldBookmark.IsFavorite", new IMAGE_BRUSH("Icons/Star_16x", Icon16x16));
+	Set("WorldPartition.WorldBookmark.IsNotFavorite", new IMAGE_BRUSH("Icons/EmptyStar_16x", Icon16x16));
+	Set("WorldPartition.WorldBookmark.RecentlyUsed", new IMAGE_BRUSH("Icons/Clock_24x", Icon16x16));
+
 	Set("WorldPartition.FocusSelection", new IMAGE_BRUSH_SVG("WorldPartition/FramingSelection_20", Icon16x16));
 	Set("WorldPartition.FocusLoadedRegions", new IMAGE_BRUSH_SVG("WorldPartition/FrameRegions_20", Icon16x16));
 	Set("WorldPartition.FocusWorld", new IMAGE_BRUSH_SVG("WorldPartition/FrameWorld_20", Icon16x16));
@@ -6275,6 +6284,7 @@ void FStarshipEditorStyle::FStyle::SetupClassIconsAndThumbnails()
 			{TEXT("GeometryCollection")},
 			{TEXT("GroupActor")},
 			{TEXT("HierarchicalInstancedStaticMeshComponent")},
+			{TEXT("HLODLayer")},
 			{TEXT("HUD")},
 			{TEXT("ImagePlate")},
 			{TEXT("InstancedStaticMeshComponent")},
@@ -6389,6 +6399,7 @@ void FStarshipEditorStyle::FStyle::SetupClassIconsAndThumbnails()
 			{TEXT("WindDirectionalSource")},
 			{TEXT("World")},
 			{TEXT("WorldDataLayers")},
+			{TEXT("WorldBookmark")},
 		};
 	
 		// SVG Asset icons
diff --git a/Engine/Source/Editor/LevelEditor/Private/LevelEditorCameraEditorState.cpp b/Engine/Source/Editor/LevelEditor/Private/LevelEditorCameraEditorState.cpp
new file mode 100644
--- /dev/null
+++ b/Engine/Source/Editor/LevelEditor/Private/LevelEditorCameraEditorState.cpp
@@ -0,0 +1,60 @@
+// Copyright Epic Games, Inc. All Rights Reserved.
+
+#include "LevelEditorCameraEditorState.h"
+#include "LevelEditorInternalTools.h"
+#include "LevelEditorViewport.h"
+
+#include UE_INLINE_GENERATED_CPP_BY_NAME(LevelEditorCameraEditorState)
+
+#define LOCTEXT_NAMESPACE "LevelEditorCameraEditorState"
+
+ULevelEditorCameraEditorState::ULevelEditorCameraEditorState(const FObjectInitializer& ObjectInitializer)
+	: Super(ObjectInitializer)
+	, CameraFOVAngle(90.0f)
+{
+}
+
+FText ULevelEditorCameraEditorState::GetCategoryText() const
+{
+	return FText(LOCTEXT("LevelEditorCameraEditorStateCategoryText", "Camera"));
+}
+
+bool ULevelEditorCameraEditorState::CaptureState()
+{
+	TSharedPtr<SLevelViewport> LevelViewport = InternalEditorLevelLibrary::GetActiveLevelViewport();
+	if (LevelViewport.IsValid())
+	{
+		const FLevelEditorViewportClient& LevelViewportClient = LevelViewport->GetLevelViewportClient();
+		CameraLocation = LevelViewportClient.GetViewLocation();
+		CameraRotation = LevelViewportClient.GetViewRotation();
+		CameraFOVAngle = LevelViewportClient.FOVAngle;
+		return true;
+	}
+	else
+	{
+		return false;
+	}
+}
+
+bool ULevelEditorCameraEditorState::RestoreState() const
+{
+	// Set all level editing cameras to this bookmark
+	for (FLevelEditorViewportClient* LevelVC : GEditor->GetLevelViewportClients())
+	{
+		LevelVC->SetViewLocation(CameraLocation);
+		if (!LevelVC->IsOrtho())
+		{
+			LevelVC->SetViewRotation(CameraRotation);
+		}
+
+		LevelVC->FOVAngle = CameraFOVAngle;
+		LevelVC->ViewFOV = CameraFOVAngle;
+		LevelVC->Invalidate();
+
+		FEditorDelegates::OnEditorCameraMoved.Broadcast(CameraLocation, CameraRotation, LevelVC->ViewportType, LevelVC->ViewIndex);
+	}
+
+	return true;
+}
+
+#undef LOCTEXT_NAMESPACE
diff --git a/Engine/Source/Editor/LevelEditor/Private/LevelEditorInternalTools.cpp b/Engine/Source/Editor/LevelEditor/Private/LevelEditorInternalTools.cpp
new file mode 100644
--- /dev/null
+++ b/Engine/Source/Editor/LevelEditor/Private/LevelEditorInternalTools.cpp
@@ -0,0 +1,63 @@
+// Copyright Epic Games, Inc. All Rights Reserved.
+
+#include "LevelEditorInternalTools.h"
+
+namespace InternalEditorLevelLibrary
+{
+	TSharedPtr<SLevelViewport> GetLevelViewport(const FName& ViewportConfigKey)
+	{
+		FLevelEditorModule& LevelEditorModule = FModuleManager::GetModuleChecked<FLevelEditorModule>("LevelEditor");
+
+		TSharedPtr<ILevelEditor> LevelEditor = LevelEditorModule.GetFirstLevelEditor();
+		if (LevelEditor.IsValid())
+		{
+			if (ViewportConfigKey != NAME_None)
+			{
+				for (TSharedPtr<SLevelViewport> LevelViewport : LevelEditor->GetViewports())
+				{
+					if (LevelViewport.IsValid() && LevelViewport->GetConfigKey() == ViewportConfigKey)
+					{
+						return LevelViewport;
+					}
+				}
+			}
+
+			TSharedPtr<SLevelViewport> ActiveLevelViewport = LevelEditor->GetActiveViewportInterface();
+			return ActiveLevelViewport;
+		}
+
+		return nullptr;
+	}
+
+	TSharedPtr<SLevelViewport> GetActiveLevelViewport()
+	{
+		return GetLevelViewport(NAME_None);
+	}
+
+	bool IsEditingLevelInstanceCurrentLevel(UWorld* InWorld)
+	{
+		if (InWorld)
+		{
+			if (ULevelInstanceSubsystem* LevelInstanceSubsystem = InWorld->GetSubsystem<ULevelInstanceSubsystem>())
+			{
+				if (ILevelInstanceInterface* LevelInstance = LevelInstanceSubsystem->GetEditingLevelInstance())
+				{
+					return LevelInstanceSubsystem->GetLevelInstanceLevel(LevelInstance) == InWorld->GetCurrentLevel();
+				}
+			}
+		}
+
+		return false;
+	}
+
+	AActor* GetEditingLevelInstance(UWorld* InWorld)
+	{
+		ULevelInstanceSubsystem* LevelInstanceSubsystem = InWorld ? InWorld->GetSubsystem<ULevelInstanceSubsystem>() : nullptr;
+		return LevelInstanceSubsystem ? Cast<AActor>(LevelInstanceSubsystem->GetEditingLevelInstance()) : nullptr;
+	}
+
+	bool IsActorEditorContextVisible(UWorld* InWorld)
+	{
+		return InWorld && (InWorld->GetCurrentLevel()->OwningWorld->GetLevels().Num() > 1 && (!InWorld->IsPartitionedWorld() || GetEditingLevelInstance(InWorld) != nullptr));
+	}
+}
diff --git a/Engine/Source/Editor/LevelEditor/Private/LevelEditorInternalTools.h b/Engine/Source/Editor/LevelEditor/Private/LevelEditorInternalTools.h
new file mode 100644
--- /dev/null
+++ b/Engine/Source/Editor/LevelEditor/Private/LevelEditorInternalTools.h
@@ -0,0 +1,19 @@
+// Copyright Epic Games, Inc. All Rights Reserved.
+
+#pragma once
+
+#include "LevelEditor.h"
+#include "Modules/ModuleManager.h"
+#include "SLevelViewport.h"
+
+#include "LevelInstance/LevelInstanceSubsystem.h"
+#include "LevelInstance/LevelInstanceInterface.h"
+
+namespace InternalEditorLevelLibrary
+{
+	TSharedPtr<SLevelViewport> GetLevelViewport(const FName& ViewportConfigKey);
+	TSharedPtr<SLevelViewport> GetActiveLevelViewport();
+	bool IsEditingLevelInstanceCurrentLevel(UWorld* InWorld);
+	AActor* GetEditingLevelInstance(UWorld* InWorld);
+	bool IsActorEditorContextVisible(UWorld* InWorld);
+}
diff --git a/Engine/Source/Editor/LevelEditor/Private/LevelEditorSubsystem.cpp b/Engine/Source/Editor/LevelEditor/Private/LevelEditorSubsystem.cpp
--- a/Engine/Source/Editor/LevelEditor/Private/LevelEditorSubsystem.cpp
+++ b/Engine/Source/Editor/LevelEditor/Private/LevelEditorSubsystem.cpp
@@ -26,65 +26,14 @@
 #include "LightingBuildOptions.h"
 #include "Widgets/Images/SImage.h"
 #include "Widgets/Input/SComboButton.h"
+#include "LevelEditorInternalTools.h"
+#include "EditorState/EditorStateSubsystem.h"
+#include "LevelEditorCameraEditorState.h"
 
 DEFINE_LOG_CATEGORY_STATIC(LevelEditorSubsystem, Log, All);
 
 #define LOCTEXT_NAMESPACE "LevelEditorSubsystem"
 
-namespace InternalEditorLevelLibrary
-{
-
-TSharedPtr<SLevelViewport> GetLevelViewport(const FName& ViewportConfigKey)
-{
-	FLevelEditorModule& LevelEditorModule = FModuleManager::GetModuleChecked<FLevelEditorModule>("LevelEditor");
-
-	TSharedPtr<ILevelEditor> LevelEditor = LevelEditorModule.GetFirstLevelEditor();
-	if (LevelEditor.IsValid())
-	{
-		for (TSharedPtr<SLevelViewport> LevelViewport : LevelEditor->GetViewports())
-		{
-			if (LevelViewport.IsValid() && LevelViewport->GetConfigKey() == ViewportConfigKey)
-			{
-				return LevelViewport;
-			}
-		}
-
-		TSharedPtr<SLevelViewport> ActiveLevelViewport = LevelEditor->GetActiveViewportInterface();
-		return ActiveLevelViewport;
-	}
-
-	return nullptr;
-}
-
-bool IsEditingLevelInstanceCurrentLevel(UWorld* InWorld)
-{
-	if (InWorld)
-	{
-		if (ULevelInstanceSubsystem* LevelInstanceSubsystem = InWorld->GetSubsystem<ULevelInstanceSubsystem>())
-		{
-			if (ILevelInstanceInterface* LevelInstance = LevelInstanceSubsystem->GetEditingLevelInstance())
-			{
-				return LevelInstanceSubsystem->GetLevelInstanceLevel(LevelInstance) == InWorld->GetCurrentLevel();
-			}
-		}
-	}
-
-	return false;
-}
-
-AActor* GetEditingLevelInstance(UWorld* InWorld)
-{
-	ULevelInstanceSubsystem* LevelInstanceSubsystem = InWorld ? InWorld->GetSubsystem<ULevelInstanceSubsystem>() : nullptr;
-	return LevelInstanceSubsystem ? Cast<AActor>(LevelInstanceSubsystem->GetEditingLevelInstance()) : nullptr;
-}
-
-bool IsActorEditorContextVisible(UWorld* InWorld)
-{
-	return InWorld && (InWorld->GetCurrentLevel()->OwningWorld->GetLevels().Num() > 1 && (!InWorld->IsPartitionedWorld() || GetEditingLevelInstance(InWorld) != nullptr));
-}
-
-}
-
 void ULevelEditorSubsystem::Initialize(FSubsystemCollectionBase& Collection)
 {
 	Collection.InitializeDependency<UActorEditorContextSubsystem>();
@@ -102,6 +51,9 @@ void ULevelEditorSubsystem::Initialize(FSubsystemCollectionBase& Collection)
 
 	FEditorDelegates::MapChange.AddUObject(this, &ULevelEditorSubsystem::HandleOnMapChanged);
 	FEditorDelegates::OnMapOpened.AddUObject(this, &ULevelEditorSubsystem::HandleOnMapOpened);
+
+	Collection.InitializeDependency<UEditorStateSubsystem>();
+	UEditorStateSubsystem::RegisterEditorStateType<ULevelEditorCameraEditorState>();
 }
 
 void ULevelEditorSubsystem::Deinitialize()
@@ -120,6 +72,8 @@ void ULevelEditorSubsystem::Deinitialize()
 
 	FEditorDelegates::MapChange.RemoveAll(this);
 	FEditorDelegates::OnMapOpened.RemoveAll(this);
+
+	UEditorStateSubsystem::UnregisterEditorStateType<ULevelEditorCameraEditorState>();
 }
 
 void ULevelEditorSubsystem::ExtendQuickActionMenu()
diff --git a/Engine/Source/Editor/LevelEditor/Public/LevelEditorCameraEditorState.h b/Engine/Source/Editor/LevelEditor/Public/LevelEditorCameraEditorState.h
new file mode 100644
--- /dev/null
+++ b/Engine/Source/Editor/LevelEditor/Public/LevelEditorCameraEditorState.h
@@ -0,0 +1,32 @@
+// Copyright Epic Games, Inc. All Rights Reserved.
+
+#pragma once
+
+#include "EditorState/EditorState.h"
+#include "LevelEditorCameraEditorState.generated.h"
+
+UCLASS()
+class ULevelEditorCameraEditorState : public UEditorState
+{
+	GENERATED_UCLASS_BODY()
+
+	//~ Begin UEditorState Interface
+	virtual FText GetCategoryText() const override;
+private:
+	virtual bool CaptureState() override;
+	virtual bool RestoreState() const override;
+	//~ End UEditorState Interface
+
+private:
+	// Camera position
+	UPROPERTY(EditAnywhere, Category = Camera)
+	FVector CameraLocation;
+
+	// Camera rotation
+	UPROPERTY(EditAnywhere, Category = Camera)
+	FRotator CameraRotation;
+
+	// Camera FOV Angle
+	UPROPERTY(EditAnywhere, Category = Camera)
+	float CameraFOVAngle;
+};
diff --git a/Engine/Source/Editor/UnrealEd/Private/EditorEngine.cpp b/Engine/Source/Editor/UnrealEd/Private/EditorEngine.cpp
--- a/Engine/Source/Editor/UnrealEd/Private/EditorEngine.cpp
+++ b/Engine/Source/Editor/UnrealEd/Private/EditorEngine.cpp
@@ -1497,6 +1497,7 @@ void UEditorEngine::LoadDefaultEditorModules()
 			TEXT("TraceInsights"),
 			TEXT("StaticMeshEditor"),
 			TEXT("EditorFramework"),
+			TEXT("WorldBookmark"),
 			TEXT("WorldPartitionEditor"),
 			TEXT("EditorConfig"),
 			TEXT("DerivedDataEditor"),
diff --git a/Engine/Source/Editor/UnrealEd/Private/EditorState/ActorEditorContextEditorState.cpp b/Engine/Source/Editor/UnrealEd/Private/EditorState/ActorEditorContextEditorState.cpp
new file mode 100644
--- /dev/null
+++ b/Engine/Source/Editor/UnrealEd/Private/EditorState/ActorEditorContextEditorState.cpp
@@ -0,0 +1,82 @@
+// Copyright Epic Games, Inc. All Rights Reserved.
+
+#include "EditorState/ActorEditorContextEditorState.h"
+
+#include "Engine/Level.h"
+#include "Engine/World.h"
+#include "GameFramework/Actor.h"
+#include "Subsystems/ActorEditorContextSubsystem.h"
+#include "UObject\Package.h"
+
+#include UE_INLINE_GENERATED_CPP_BY_NAME(ActorEditorContextEditorState)
+
+#define LOCTEXT_NAMESPACE "ActorContextEditorState"
+
+UActorEditorContextEditorState::UActorEditorContextEditorState(const FObjectInitializer& ObjectInitializer)
+	: Super(ObjectInitializer)
+{
+}
+
+void UActorEditorContextEditorState::PostInitProperties()
+{
+	Super::PostInitProperties();
+
+	if (!HasAnyFlags(RF_ClassDefaultObject | RF_NeedLoad))
+	{
+		ActorEditorContext = NewObject<AActor>(this);
+	}
+}
+
+void UActorEditorContextEditorState::PostDuplicate(bool bDuplicateForPIE)
+{
+	Super::PostDuplicate(bDuplicateForPIE);
+
+	ActorEditorContext = NewObject<AActor>(this);
+}
+
+FText UActorEditorContextEditorState::GetCategoryText() const
+{
+	return FText(LOCTEXT("ActorEditorContextEditorStateCategoryText", "Actor Editor Context"));
+}
+
+bool UActorEditorContextEditorState::CaptureState()
+{
+	UWorld* CurrentWorld = GetStateWorld();
+
+	// Store context to saved actor with EActorEditorContextAction::ApplyContext
+	ActorEditorContext->SetFlags(RF_Transient);
+	ActorEditorContext->Rename(nullptr, GetTransientPackage(), REN_DoNotDirty | REN_DontCreateRedirectors | REN_NonTransactional);
+	ActorEditorContext->MarkAsGarbage();
+
+	UPackage* TempExternalActorPackage = CreatePackage(*FString::Printf(TEXT("/Temp/ActorEditorContextEditorState/%s"), *FGuid::NewGuid().ToString()));
+	TempExternalActorPackage->SetFlags(RF_Transient);
+
+	ActorEditorContext = NewObject<AActor>(CurrentWorld->PersistentLevel);
+	ActorEditorContext->SetExternalPackage(TempExternalActorPackage);
+
+	UActorEditorContextSubsystem::Get()->ApplyContext(ActorEditorContext);
+
+	ActorEditorContext->SetExternalPackage(nullptr);
+	ActorEditorContext->Rename(nullptr, this, REN_DoNotDirty | REN_DontCreateRedirectors | REN_NonTransactional);
+
+	return true;
+}
+
+bool UActorEditorContextEditorState::RestoreState() const
+{
+	UWorld* CurrentWorld = GetStateWorld();
+
+	// Actor Editor Context - Apply context from saved actor with EActorEditorContextAction::InitializeContextFromActor
+	ActorEditorContext->Rename(nullptr, CurrentWorld->PersistentLevel, REN_DoNotDirty | REN_DontCreateRedirectors | REN_NonTransactional);
+	ActorEditorContext->SetExternalPackage(GetPackage());
+
+	UActorEditorContextSubsystem::Get()->ResetContext();
+	UActorEditorContextSubsystem::Get()->InitializeContextFromActor(ActorEditorContext);
+
+	ActorEditorContext->SetExternalPackage(nullptr);
+	ActorEditorContext->Rename(nullptr, const_cast<UActorEditorContextEditorState*>(this), REN_DoNotDirty | REN_DontCreateRedirectors | REN_NonTransactional);
+
+	return true;
+}
+
+#undef LOCTEXT_NAMESPACE
diff --git a/Engine/Source/Editor/UnrealEd/Private/EditorState/EditorState.cpp b/Engine/Source/Editor/UnrealEd/Private/EditorState/EditorState.cpp
new file mode 100644
--- /dev/null
+++ b/Engine/Source/Editor/UnrealEd/Private/EditorState/EditorState.cpp
@@ -0,0 +1,5 @@
+// Copyright Epic Games, Inc. All Rights Reserved.
+
+#include "EditorState/EditorState.h"
+
+#include UE_INLINE_GENERATED_CPP_BY_NAME(EditorState)
diff --git a/Engine/Source/Editor/UnrealEd/Private/EditorState/EditorStateCollection.cpp b/Engine/Source/Editor/UnrealEd/Private/EditorState/EditorStateCollection.cpp
new file mode 100644
--- /dev/null
+++ b/Engine/Source/Editor/UnrealEd/Private/EditorState/EditorStateCollection.cpp
@@ -0,0 +1,5 @@
+// Copyright Epic Games, Inc. All Rights Reserved.
+
+#include "EditorState/EditorStateCollection.h"
+
+#include UE_INLINE_GENERATED_CPP_BY_NAME(EditorStateCollection)
diff --git a/Engine/Source/Editor/UnrealEd/Private/EditorState/EditorStateSubsystem.cpp b/Engine/Source/Editor/UnrealEd/Private/EditorState/EditorStateSubsystem.cpp
new file mode 100644
--- /dev/null
+++ b/Engine/Source/Editor/UnrealEd/Private/EditorState/EditorStateSubsystem.cpp
@@ -0,0 +1,101 @@
+// Copyright Epic Games, Inc. All Rights Reserved.
+
+#include "EditorState/EditorStateSubsystem.h"
+#include "EditorState/WorldEditorState.h"
+#include "Editor.h"
+
+#include UE_INLINE_GENERATED_CPP_BY_NAME(EditorStateSubsystem)
+
+UEditorStateSubsystem* UEditorStateSubsystem::Get()
+{
+	check(GEditor);
+	return GEditor->GetEditorSubsystem<UEditorStateSubsystem>();
+}
+
+void UEditorStateSubsystem::Initialize(FSubsystemCollectionBase& Collection)
+{
+	Super::Initialize(Collection);
+	UEditorStateSubsystem::RegisterEditorStateType<UWorldEditorState>();
+}
+
+void UEditorStateSubsystem::Deinitialize()
+{
+	UEditorStateSubsystem::UnregisterEditorStateType<UWorldEditorState>();
+	Super::Deinitialize();
+}
+
+void UEditorStateSubsystem::RegisterEditorStateType(TSubclassOf<UEditorState> InEditorStateType)
+{
+	check(!RegisteredEditorStateTypes.Contains(InEditorStateType));
+	RegisteredEditorStateTypes.Add(InEditorStateType);
+}
+
+void UEditorStateSubsystem::UnregisterEditorStateType(TSubclassOf<UEditorState> InEditorStateType)
+{
+	check(RegisteredEditorStateTypes.Contains(InEditorStateType));
+	RegisteredEditorStateTypes.RemoveSwap(InEditorStateType);
+}
+	
+void UEditorStateSubsystem::CaptureEditorState(FEditorStateCollection& OutState, UObject* InStateOuter)
+{
+	// Keep track of restored states to avoid cycles
+	TSet<TSubclassOf<UEditorState>> ProcessedStates;
+
+	// Process each state, but capture dependencies first
+	const auto CaptureState = [&ProcessedStates, &OutState](UEditorState* StateToCapture)
+	{
+		auto CaptureStateImpl = [&ProcessedStates, &OutState](UEditorState* StateToCapture, auto& CaptureStateRef) -> void
+		{
+			if (!ProcessedStates.Contains(StateToCapture->GetClass()))
+			{
+				for (TSubclassOf<UEditorState> DependencyType : StateToCapture->GetDependencies())
+				{
+					CaptureStateRef(OutState.EditorStates.FindChecked(DependencyType), CaptureStateRef);
+				}
+
+				ProcessedStates.Add(StateToCapture->GetClass());
+				StateToCapture->CaptureState();
+			}
+		};
+
+		CaptureStateImpl(StateToCapture, CaptureStateImpl);
+	};
+
+
+	OutState.EditorStates.Reset();
+
+	for (TSubclassOf<UEditorState> EditorStateType : RegisteredEditorStateTypes)
+	{
+		OutState.EditorStates.Emplace(EditorStateType, NewObject<UEditorState>(InStateOuter, EditorStateType));
+	}
+
+	OutState.ForEachState(CaptureState);
+}
+
+void UEditorStateSubsystem::RestoreEditorState(const FEditorStateCollection& InState)
+{
+	// Keep track of restored states to avoid cycles
+	TSet<TSubclassOf<UEditorState>> ProcessedStates;	
+	
+	// Process each state, but restore dependencies first
+	const auto RestoreState = [&ProcessedStates, &InState](const UEditorState* StateToRestore)
+	{
+		auto RestoreStateImpl = [&ProcessedStates, &InState](const UEditorState* StateToRestore, auto& RestoreStateRef) -> void
+		{
+			if (!ProcessedStates.Contains(StateToRestore->GetClass()))
+			{
+				for (TSubclassOf<UEditorState> DependencyType : StateToRestore->GetDependencies())
+				{
+					RestoreStateRef(InState.EditorStates.FindChecked(DependencyType), RestoreStateRef);
+				}
+
+				ProcessedStates.Add(StateToRestore->GetClass());
+				StateToRestore->RestoreState();
+			}
+		};
+
+		RestoreStateImpl(StateToRestore, RestoreStateImpl);
+	};
+
+	InState.ForEachState(RestoreState);
+}
diff --git a/Engine/Source/Editor/UnrealEd/Private/EditorState/WorldEditorState.cpp b/Engine/Source/Editor/UnrealEd/Private/EditorState/WorldEditorState.cpp
new file mode 100644
--- /dev/null
+++ b/Engine/Source/Editor/UnrealEd/Private/EditorState/WorldEditorState.cpp
@@ -0,0 +1,98 @@
+// Copyright Epic Games, Inc. All Rights Reserved.
+
+#include "EditorState/WorldEditorState.h"
+#include "Engine/World.h"
+#include "FileHelpers.h"
+#include "Misc/MessageDialog.h"
+
+#include UE_INLINE_GENERATED_CPP_BY_NAME(WorldEditorState)
+
+#define LOCTEXT_NAMESPACE "WorldEditorState"
+
+UWorldEditorState::UWorldEditorState(const FObjectInitializer& ObjectInitializer)
+	: Super(ObjectInitializer)
+{
+}
+
+FText UWorldEditorState::GetCategoryText() const
+{
+	return FText(LOCTEXT("WorldEditorStateCategoryText", "World"));
+}
+
+TSoftObjectPtr<UWorld> UWorldEditorState::GetStateWorld() const
+{
+	return World;
+}
+
+bool UWorldEditorState::CaptureState()
+{
+	// Set the current world.
+	World = (UWorld*)GWorld;
+	return true;
+}
+
+bool UWorldEditorState::RestoreState() const
+{
+	if (World.IsNull())
+	{
+		return false;
+	}
+
+	UWorld* CurrentWorld = GWorld;
+	UWorld* TargetWorld = Cast<UWorld>(World.ToSoftObjectPath().ResolveObject());
+
+	// Already in the right world.
+	if (CurrentWorld == TargetWorld)
+	{
+		return true;
+	}
+
+	// If there are any unsaved changes to the current level, see if the user wants to save those first
+	// If they do not wish to save, then we will bail out of opening this asset.
+	constexpr bool bPromptUserToSave = true;
+	constexpr bool bSaveMapPackages = true;
+	constexpr bool bSaveContentPackages = true;
+	if (FEditorFileUtils::SaveDirtyPackages(bPromptUserToSave, bSaveMapPackages, bSaveContentPackages))
+	{
+		if (TargetWorld != nullptr)
+		{
+			// Validate that Asset was saved or isn't loaded meaning it can be loaded
+			if (TargetWorld->GetPackage() == nullptr)
+			{
+				FMessageDialog::Open(EAppMsgType::Ok, FText::Format(LOCTEXT("CannotOpenNotInPackage", "World ({0}) is not in a package and cannot be opened."), FText::FromString(TargetWorld->GetFullName())));
+				return false;
+			}
+
+			if (TargetWorld->GetPackage()->HasAnyPackageFlags(PKG_NewlyCreated))
+			{
+				FMessageDialog::Open(EAppMsgType::Ok, LOCTEXT("CannotOpenNewlyCreatedMapWithoutSaving", "The world you are trying to open needs to be saved first."));
+				return false;
+			}
+		}
+	}
+	else
+	{
+		return false;
+	}
+
+	const bool bLoadAsTemplate = false;
+	const bool bShowProgress = true;
+	return FEditorFileUtils::LoadMap(FPackageName::LongPackageNameToFilename(World.GetLongPackageName(), FPackageName::GetMapPackageExtension()), bLoadAsTemplate, bShowProgress);
+}
+	
+UWorldDependantEditorState::UWorldDependantEditorState(const FObjectInitializer& ObjectInitializer)
+	: Super(ObjectInitializer)
+{
+}
+
+UWorld* UWorldDependantEditorState::GetStateWorld() const
+{
+	return GWorld;
+}
+
+TArray<TSubclassOf<UEditorState>> UWorldDependantEditorState::GetDependencies() const
+{
+	return { UWorldEditorState::StaticClass() };
+}
+
+#undef LOCTEXT_NAMESPACE
diff --git a/Engine/Source/Editor/UnrealEd/Private/Subsystems/ActorEditorContextSubsystem.cpp b/Engine/Source/Editor/UnrealEd/Private/Subsystems/ActorEditorContextSubsystem.cpp
--- a/Engine/Source/Editor/UnrealEd/Private/Subsystems/ActorEditorContextSubsystem.cpp
+++ b/Engine/Source/Editor/UnrealEd/Private/Subsystems/ActorEditorContextSubsystem.cpp
@@ -7,6 +7,8 @@
 #include "ScopedTransaction.h"
 #include "UnrealEdGlobals.h"
 #include "Editor.h"
+#include "EditorState/EditorStateSubsystem.h"
+#include "EditorState/ActorEditorContextEditorState.h"
 
 #define LOCTEXT_NAMESPACE "ActorEditorContext"
 
@@ -26,6 +28,9 @@ void UActorEditorContextSubsystem::Initialize(FSubsystemCollectionBase& Collecti
 		GUnrealEd->OnPasteActorsBegin().AddUObject(this, &UActorEditorContextSubsystem::OnPasteActorsBegin);
 		GUnrealEd->OnPasteActorsEnd().AddUObject(this, &UActorEditorContextSubsystem::OnPasteActorsEnd);
 	}
+
+	Collection.InitializeDependency<UEditorStateSubsystem>();
+	UEditorStateSubsystem::RegisterEditorStateType<UActorEditorContextEditorState>();
 }
 
 void UActorEditorContextSubsystem::Deinitialize()
@@ -38,6 +43,8 @@ void UActorEditorContextSubsystem::Deinitialize()
 		GUnrealEd->OnPasteActorsEnd().RemoveAll(this);
 	}
 
+	UEditorStateSubsystem::UnregisterEditorStateType<UActorEditorContextEditorState>();
+
 	Super::Deinitialize();
 }
 
diff --git a/Engine/Source/Editor/UnrealEd/Private/Subsystems/AssetEditorSubsystem.cpp b/Engine/Source/Editor/UnrealEd/Private/Subsystems/AssetEditorSubsystem.cpp
--- a/Engine/Source/Editor/UnrealEd/Private/Subsystems/AssetEditorSubsystem.cpp
+++ b/Engine/Source/Editor/UnrealEd/Private/Subsystems/AssetEditorSubsystem.cpp
@@ -515,12 +515,12 @@ bool UAssetEditorSubsystem::OpenEditorForAsset(UObject* Asset, const EToolkitMod
 	//    editor and summon it in the new level editor, or to just foreground the old level editor (current behavior)
 
 	FAssetToolsModule& AssetToolsModule = FModuleManager::LoadModuleChecked<FAssetToolsModule>(TEXT("AssetTools"));
-
-	const bool bBringToFrontIfOpen = true;
+	TWeakPtr<IAssetTypeActions> AssetTypeActions = AssetToolsModule.Get().GetAssetTypeActionsForClass(Asset->GetClass());
 
 	AssetEditorRequestOpenEvent.Broadcast(Asset);
 
-	if (FindEditorForAsset(Asset, bBringToFrontIfOpen) != nullptr)
+	const bool bBringToFrontIfOpen = true;
+	if ((AssetTypeActions.IsValid() && AssetTypeActions.Pin()->ShouldFindEditorForAsset()) && FindEditorForAsset(Asset, bBringToFrontIfOpen) != nullptr)
 	{
 		// This asset is already open in an editor! (the call to FindEditorForAsset above will bring it to the front)
 		return true;
@@ -541,8 +541,6 @@ bool UAssetEditorSubsystem::OpenEditorForAsset(UObject* Asset, const EToolkitMod
 	// We store the open method for this asset, so that FAssetEditorToolkit can query us for this information during init
 	AssetOpenMethodCache.Add(AssetPath, AssetOpenMethod);
 
-	TWeakPtr<IAssetTypeActions> AssetTypeActions = AssetToolsModule.Get().GetAssetTypeActionsForClass(Asset->GetClass());
-
 	EToolkitMode::Type ActualToolkitMode = ToolkitMode;
 	if (AssetTypeActions.IsValid())
 	{
diff --git a/Engine/Source/Editor/UnrealEd/Public/EditorState/ActorEditorContextEditorState.h b/Engine/Source/Editor/UnrealEd/Public/EditorState/ActorEditorContextEditorState.h
new file mode 100644
--- /dev/null
+++ b/Engine/Source/Editor/UnrealEd/Public/EditorState/ActorEditorContextEditorState.h
@@ -0,0 +1,32 @@
+// Copyright Epic Games, Inc. All Rights Reserved.
+
+#pragma once
+
+#include "EditorState/WorldEditorState.h"
+#include "ActorEditorContextEditorState.generated.h"
+
+class AActor;
+
+UCLASS(MinimalAPI)
+class UActorEditorContextEditorState : public UWorldDependantEditorState
+{
+	GENERATED_UCLASS_BODY()
+
+public:
+	//~ Begin UEditorState Interface
+	virtual FText GetCategoryText() const override;
+
+protected:
+	virtual bool CaptureState() override;
+	virtual bool RestoreState() const override;
+	//~ End UEditorState Interface
+
+	//~ Begin UObject Interface
+	virtual void PostInitProperties() override;
+	virtual void PostDuplicate(bool bDuplicateForPIE) override;
+	//~ End UObject Interface
+
+	// Actor used to store the editor context, will be used to restore context when loading the bookmark.
+	UPROPERTY(DuplicateTransient)
+	TObjectPtr<AActor> ActorEditorContext;
+};
diff --git a/Engine/Source/Editor/UnrealEd/Public/EditorState/EditorState.h b/Engine/Source/Editor/UnrealEd/Public/EditorState/EditorState.h
new file mode 100644
--- /dev/null
+++ b/Engine/Source/Editor/UnrealEd/Public/EditorState/EditorState.h
@@ -0,0 +1,41 @@
+// Copyright Epic Games, Inc. All Rights Reserved.
+
+#pragma once
+
+#include "UObject/Object.h"
+#include "Templates/SubclassOf.h"
+#include "EditorState.generated.h"
+
+/**
+ * EditorState is a container that can capture the state of a given editor subsystem and restore it at a later time.
+ * To capture/restore states, you must use the UEditorStateSubsystem subsystem rather than dealing with this class directly.
+ * @see UEditorStateSubsystem
+ */
+UCLASS(MinimalAPI, Abstract)
+class UEditorState : public UObject
+{
+	GENERATED_BODY()
+
+public:
+	/** Get the category under which this state's properties should be displayed. */
+	UNREALED_API virtual FText GetCategoryText() const PURE_VIRTUAL(UEditorState::GetCategoryText, return FText::GetEmpty(); );
+	
+protected:
+	/** Get a list of state types this editor state depends on. Dependent types will be restored after their dependencies. */
+	virtual TArray<TSubclassOf<UEditorState>> GetDependencies() const { return {}; }
+
+	/** Capture the state of the editor. Must be subclassed. */
+	virtual bool CaptureState() PURE_VIRTUAL(UEditorState::CaptureState, return false; );
+
+	/** Restore the state of the editor. Must be subclassed. */
+	virtual bool RestoreState() const PURE_VIRTUAL(UEditorState::RestoreState, return false; );
+
+private:
+	/** Prevent access to UObject::GetWorld() as it's not relevant for this class and is error prone. */
+	virtual class UWorld* GetWorld() const override final 
+	{
+		return nullptr;
+	}
+
+	friend class UEditorStateSubsystem;
+};
diff --git a/Engine/Source/Editor/UnrealEd/Public/EditorState/EditorStateCollection.h b/Engine/Source/Editor/UnrealEd/Public/EditorState/EditorStateCollection.h
new file mode 100644
--- /dev/null
+++ b/Engine/Source/Editor/UnrealEd/Public/EditorState/EditorStateCollection.h
@@ -0,0 +1,70 @@
+// Copyright Epic Games, Inc. All Rights Reserved.
+
+#pragma once
+
+#include "EditorState/EditorState.h"
+#include "Templates/EnableIf.h"
+#include "EditorStateCollection.generated.h"
+
+/**
+ * An immutable collection of editor state objects that can be queried and iterated upon.
+  */
+USTRUCT()
+struct FEditorStateCollection
+{
+	GENERATED_BODY()
+
+	template <typename TEditorStateType>
+	typename TEnableIf<TIsDerivedFrom<TEditorStateType, UEditorState>::IsDerived, bool>::Type HasState() const
+	{
+		return EditorStates.Contains(TEditorStateType::StaticClass());
+	}
+
+	template <typename TEditorStateType>
+	typename TEnableIf<TIsDerivedFrom<TEditorStateType, UEditorState>::IsDerived, const TEditorStateType*>::Type GetState() const
+	{
+		const TObjectPtr<UEditorState>* EditorStatePtr = EditorStates.Find(TEditorStateType::StaticClass());
+		return EditorStatePtr ? CastChecked<TEditorStateType>((*EditorStatePtr).Get()) : nullptr;
+	}
+
+	template <typename TEditorStateType>
+	typename TEnableIf<TIsDerivedFrom<TEditorStateType, UEditorState>::IsDerived, const TEditorStateType&>::Type GetStateChecked() const
+	{
+		const TObjectPtr<UEditorState>& EditorState = EditorStates.FindChecked(TEditorStateType::StaticClass());
+		return *CastChecked<TEditorStateType>(EditorState.Get());
+	}
+
+	void ForEachState(TFunctionRef<void(const UEditorState*)> InFunc) const
+	{
+		for (const TPair<TSubclassOf<UEditorState>, TObjectPtr<UEditorState>>& KeyValue : EditorStates)
+		{
+			InFunc(KeyValue.Value);
+		}
+	}
+
+	TArray<UEditorState*> GetStates() const
+	{
+		TArray<UEditorState*> Values;
+
+		for (const TPair<TSubclassOf<UEditorState>, TObjectPtr<UEditorState>>& KeyValue : EditorStates)
+		{
+			Values.Add(KeyValue.Value);
+		}
+	
+		return Values;
+	}
+
+private:
+	void ForEachState(TFunctionRef<void(UEditorState*)> InFunc)
+	{
+		for (const TPair<TSubclassOf<UEditorState>, TObjectPtr<UEditorState>>& KeyValue : EditorStates)
+		{
+			InFunc(KeyValue.Value);
+		}
+	}
+
+	UPROPERTY()
+	TMap<TSubclassOf<UEditorState>, TObjectPtr<UEditorState>> EditorStates;
+	
+	friend class UEditorStateSubsystem;
+};
diff --git a/Engine/Source/Editor/UnrealEd/Public/EditorState/EditorStateSubsystem.h b/Engine/Source/Editor/UnrealEd/Public/EditorState/EditorStateSubsystem.h
new file mode 100644
--- /dev/null
+++ b/Engine/Source/Editor/UnrealEd/Public/EditorState/EditorStateSubsystem.h
@@ -0,0 +1,58 @@
+// Copyright Epic Games, Inc. All Rights Reserved.
+
+#pragma once
+
+#include "EditorSubsystem.h"
+#include "EditorState/EditorState.h"
+#include "EditorState/EditorStateCollection.h"
+#include "Templates/SubclassOf.h"
+#include "EditorStateSubsystem.generated.h"
+
+/**
+ * Subsystem that allows the capture & restoration of editor states.
+ */
+UCLASS(MinimalAPI)
+class UEditorStateSubsystem : public UEditorSubsystem
+{
+	GENERATED_BODY()
+
+public:
+	/** Utility method to retrieve the subsystem. Editor must have been initialized before this can be used. */
+	static UNREALED_API UEditorStateSubsystem* Get();
+
+	UNREALED_API virtual void Initialize(FSubsystemCollectionBase& Collection) override;
+	UNREALED_API virtual void Deinitialize() override;
+
+	/** Register an EditorState type so that it can be captured in the next CaptureEditorState() calls. */
+	template <typename TEditorStateType>
+	static typename TEnableIf<TIsDerivedFrom<TEditorStateType, UEditorState>::IsDerived, void>::Type RegisterEditorStateType()
+	{
+		UEditorStateSubsystem::Get()->RegisterEditorStateType(TEditorStateType::StaticClass());
+	}
+
+	/** Unregister a previously registed EditorState type. */
+	template <typename TEditorStateType>
+	static typename TEnableIf<TIsDerivedFrom<TEditorStateType, UEditorState>::IsDerived, void>::Type UnregisterEditorStateType()
+	{
+		UEditorStateSubsystem::Get()->UnregisterEditorStateType(TEditorStateType::StaticClass());
+	}
+	
+	/**
+	 * Captures the state of the editor using all the registered EditorState subclasses.
+	 * @param OutState		Collection of states that will be populated with state data.
+	 * @param InStateOuter	Outer to use for the UEditorState objects that may be created by this method.
+	 */
+	UNREALED_API void CaptureEditorState(FEditorStateCollection& OutState, UObject* InStateOuter);
+
+	/**
+	 * Restore the state of the editor using a state collection previously captured with CaptureEditorState().
+	 * @param InState		Collection of states to restore.
+	 */
+	UNREALED_API void RestoreEditorState(const FEditorStateCollection& InState);
+	
+private:
+	UNREALED_API void RegisterEditorStateType(TSubclassOf<UEditorState> InEditorStateType);
+	UNREALED_API void UnregisterEditorStateType(TSubclassOf<UEditorState> InEditorStateType);
+
+	TArray<TSubclassOf<UEditorState>> RegisteredEditorStateTypes;
+};
diff --git a/Engine/Source/Editor/UnrealEd/Public/EditorState/WorldEditorState.h b/Engine/Source/Editor/UnrealEd/Public/EditorState/WorldEditorState.h
new file mode 100644
--- /dev/null
+++ b/Engine/Source/Editor/UnrealEd/Public/EditorState/WorldEditorState.h
@@ -0,0 +1,41 @@
+// Copyright Epic Games, Inc. All Rights Reserved.
+
+#pragma once
+
+#include "EditorState/EditorState.h"
+#include "WorldEditorState.generated.h"
+
+class UWorld;
+
+UCLASS(MinimalAPI)
+class UWorldEditorState : public UEditorState
+{
+	GENERATED_UCLASS_BODY()
+
+	UNREALED_API TSoftObjectPtr<UWorld> GetStateWorld() const;
+
+	//~ Begin UEditorState Interface
+	virtual FText GetCategoryText() const override;
+private:
+	virtual bool CaptureState() override;
+	virtual bool RestoreState() const override;
+	//~ End UEditorState Interface
+
+	// World associated with the bookmark
+	UPROPERTY(EditAnywhere, Category = General)
+	TSoftObjectPtr<UWorld> World;
+};
+
+UCLASS(MinimalAPI, Abstract)
+class UWorldDependantEditorState : public UEditorState
+{
+	GENERATED_UCLASS_BODY()
+
+public:
+	UNREALED_API UWorld* GetStateWorld() const;
+
+protected:
+	//~ Begin UEditorState Interface
+	UNREALED_API virtual TArray<TSubclassOf<UEditorState>> GetDependencies() const override;
+	//~ End UEditorState Interface
+};
diff --git a/Engine/Source/Editor/UnrealEd/Public/Subsystems/ActorEditorContextSubsystem.h b/Engine/Source/Editor/UnrealEd/Public/Subsystems/ActorEditorContextSubsystem.h
--- a/Engine/Source/Editor/UnrealEd/Public/Subsystems/ActorEditorContextSubsystem.h
+++ b/Engine/Source/Editor/UnrealEd/Public/Subsystems/ActorEditorContextSubsystem.h
@@ -40,7 +40,7 @@ public:
 private:
 
 	UNREALED_API UWorld* GetWorld() const;
-	UNREALED_API void OnActorEditorContextClientChanged(IActorEditorContextClient* Client);
+	UNREALED_API void OnActorEditorContextClientChanged(IActorEditorContextClient* Client);	
 	UNREALED_API void ApplyContext(AActor* InActor);
 
 	void OnPasteActorsBegin();
@@ -50,4 +50,6 @@ private:
 	TArray<IActorEditorContextClient*> Clients;
 	TArray<TArray<IActorEditorContextClient*>> PushedContextsStack;
 	bool bIsApplyEnabled = true;
-};
+
+	friend class UActorEditorContextEditorState;
+};
\ No newline at end of file
diff --git a/Engine/Source/Editor/UnrealEd/UnrealEd.Build.cs b/Engine/Source/Editor/UnrealEd/UnrealEd.Build.cs
--- a/Engine/Source/Editor/UnrealEd/UnrealEd.Build.cs
+++ b/Engine/Source/Editor/UnrealEd/UnrealEd.Build.cs
@@ -235,6 +235,7 @@ public class UnrealEd : ModuleRules
 				"ClothPainter",
 				"Media",
 				"VirtualTexturingEditor",
+				"WorldBookmark",
 				"WorldPartitionEditor",
 				"CSVtoSVG",
 				"SourceControlWindowExtender",
diff --git a/Engine/Source/Editor/WorldBookmark/Private/WorldBookmark/AssetDefinition_WorldBookmark.cpp b/Engine/Source/Editor/WorldBookmark/Private/WorldBookmark/AssetDefinition_WorldBookmark.cpp
new file mode 100644
--- /dev/null
+++ b/Engine/Source/Editor/WorldBookmark/Private/WorldBookmark/AssetDefinition_WorldBookmark.cpp
@@ -0,0 +1,121 @@
+// Copyright Epic Games, Inc. All Rights Reserved.
+
+#include "WorldBookmark/AssetDefinition_WorldBookmark.h"
+#include "WorldBookmark/WorldBookmark.h"
+#include "ContentBrowserMenuContexts.h"
+#include "Toolkits/SimpleAssetEditor.h"
+#include "ToolMenus.h"
+
+#define LOCTEXT_NAMESPACE "AssetTypeActions"
+
+EAssetCommandResult UAssetDefinition_WorldBookmark::ActivateAssets(const FAssetActivateArgs& ActivateArgs) const
+{
+	bOpenFromContentBrowser = true;
+	return Super::ActivateAssets(ActivateArgs);
+}
+
+EAssetCommandResult UAssetDefinition_WorldBookmark::OpenAssets(const FAssetOpenArgs& OpenArgs) const
+{
+	const TArray<UWorldBookmark*> WorldBookmarks = OpenArgs.LoadObjects<UWorldBookmark>();
+	check(WorldBookmarks.Num() > 0);
+
+	bool bEditAsset = bOpenFromContentBrowser;
+	bOpenFromContentBrowser = false;
+
+	if (bEditAsset)
+	{
+		for (UWorldBookmark* WorldBookmark : WorldBookmarks)
+		{
+			FSimpleAssetEditor::CreateEditor(EToolkitMode::Standalone, OpenArgs.ToolkitHost, WorldBookmark);
+		}
+	}
+	else
+	{
+		UWorldBookmark* WorldBookmark = WorldBookmarks[0];
+		check(WorldBookmark != nullptr);
+		WorldBookmark->Load();
+	}
+
+	return EAssetCommandResult::Handled;
+}
+
+
+//--------------------------------------------------------------------
+// Menu Extensions
+//--------------------------------------------------------------------
+namespace MenuExtension_Bookmark
+{	
+	bool CanExecuteBookmarkAction(const FToolMenuContext& InContext)
+	{
+		if (UContentBrowserAssetContextMenuContext::GetNumAssetsSelected(InContext) != 1)
+		{
+			return false;
+		}
+
+		const UContentBrowserAssetContextMenuContext* Context = UContentBrowserAssetContextMenuContext::FindContextWithAssets(InContext);
+		const TArray<UWorldBookmark*> WorldBookmarks = Context->LoadSelectedObjects<UWorldBookmark>();
+		return WorldBookmarks.Num() == 1 && WorldBookmarks[0] != nullptr;
+	}
+
+	void ExecuteGoToWorldBookmark(const FToolMenuContext& InContext)
+	{
+		const UContentBrowserAssetContextMenuContext* Context = UContentBrowserAssetContextMenuContext::FindContextWithAssets(InContext);
+		const TArray<UWorldBookmark*> WorldBookmarks = Context->LoadSelectedObjects<UWorldBookmark>();
+		
+		if (ensure(WorldBookmarks.Num() == 1 && WorldBookmarks[0] != nullptr))
+		{
+			WorldBookmarks[0]->Load();
+		}
+	}
+
+	void ExecuteUpdateWorldBookmark(const FToolMenuContext& InContext)
+	{
+		const UContentBrowserAssetContextMenuContext* Context = UContentBrowserAssetContextMenuContext::FindContextWithAssets(InContext);
+		const TArray<UWorldBookmark*> WorldBookmarks = Context->LoadSelectedObjects<UWorldBookmark>();
+
+		if (ensure(WorldBookmarks.Num() == 1 && WorldBookmarks[0] != nullptr))
+		{			
+			WorldBookmarks[0]->Update();
+		}
+	}
+
+	static FDelayedAutoRegisterHelper DelayedAutoRegister(EDelayedRegisterRunPhase::EndOfEngineInit, []
+	{
+		UToolMenus::RegisterStartupCallback(FSimpleMulticastDelegate::FDelegate::CreateLambda([]()
+		{
+			FToolMenuOwnerScoped OwnerScoped(UE_MODULE_NAME);
+			UToolMenu* Menu = UE::ContentBrowser::ExtendToolMenu_AssetContextMenu(UWorldBookmark::StaticClass());
+
+			FToolMenuSection& Section = Menu->FindOrAddSection("GetAssetActions");
+			Section.AddDynamicEntry(NAME_None, FNewToolMenuSectionDelegate::CreateLambda([](FToolMenuSection& InSection)
+			{
+				const TAttribute<FText> Label = LOCTEXT("WorldBookmark_GoTo", "Go To");
+				const TAttribute<FText> ToolTip = LOCTEXT("WorldBookmark_GoToTooltip", "Will open the world and apply the editor state saved with this bookmark.");
+				const FSlateIcon Icon = FSlateIcon(FAppStyle::GetAppStyleSetName(), "WorldPartition.WorldBookmark.GoTo");
+
+				FToolUIAction UIAction;
+				UIAction.ExecuteAction = FToolMenuExecuteAction::CreateStatic(&ExecuteGoToWorldBookmark);
+				UIAction.CanExecuteAction = FToolMenuCanExecuteAction::CreateStatic(&CanExecuteBookmarkAction);
+
+				InSection.AddMenuEntry("Bookmark_Edit", Label, ToolTip, Icon, UIAction);
+			}));
+
+			Section.AddDynamicEntry(NAME_None, FNewToolMenuSectionDelegate::CreateLambda([](FToolMenuSection& InSection)
+			{
+				const TAttribute<FText> Label = LOCTEXT("WorldBookmark_Update", "Update");
+				const TAttribute<FText> ToolTip = LOCTEXT("WorldBookmark_UpdateTooltip", "Update the bookmark with the current loaded state of the current loaded world.");
+				const FSlateIcon Icon = FSlateIcon(FAppStyle::GetAppStyleSetName(), "WorldPartition.WorldBookmark.Update");
+					
+				FToolUIAction UIAction;
+				UIAction.ExecuteAction = FToolMenuExecuteAction::CreateStatic(&ExecuteUpdateWorldBookmark);
+				UIAction.CanExecuteAction = FToolMenuCanExecuteAction::CreateStatic(&CanExecuteBookmarkAction);
+
+				InSection.AddMenuEntry("WorldBookmark_Update", Label, ToolTip, Icon, UIAction);
+			}));
+		}));
+	});
+}
+
+
+
+
diff --git a/Engine/Source/Editor/WorldBookmark/Private/WorldBookmark/AssetDefinition_WorldBookmark.h b/Engine/Source/Editor/WorldBookmark/Private/WorldBookmark/AssetDefinition_WorldBookmark.h
new file mode 100644
--- /dev/null
+++ b/Engine/Source/Editor/WorldBookmark/Private/WorldBookmark/AssetDefinition_WorldBookmark.h
@@ -0,0 +1,36 @@
+// Copyright Epic Games, Inc. All Rights Reserved.
+
+#pragma once
+
+#include "WorldBookmark/WorldBookmark.h"
+#include "AssetDefinitionDefault.h"
+
+#include "AssetDefinition_WorldBookmark.generated.h"
+
+
+UCLASS()
+class UAssetDefinition_WorldBookmark : public UAssetDefinitionDefault
+{
+	GENERATED_BODY()
+
+public:
+	// UAssetDefinition Begin
+	virtual FText GetAssetDisplayName() const override { return NSLOCTEXT("AssetDefinition", "AssetDefinition_WorldBookmark", "Bookmark"); }
+	virtual FLinearColor GetAssetColor() const override { return FLinearColor(FColor(255, 0, 0)); }
+	virtual TSoftClassPtr<UObject> GetAssetClass() const override { return UWorldBookmark::StaticClass(); }
+	virtual FAssetSupportResponse CanLocalize(const FAssetData& InAsset) const override { return FAssetSupportResponse::NotSupported(); }	
+
+	virtual TConstArrayView<FAssetCategoryPath> GetAssetCategories() const
+	{
+		static const auto Categories = { EAssetCategoryPaths::World };
+		return Categories;
+	}
+
+	virtual bool ShouldFindEditorForAsset() const override { return bOpenFromContentBrowser; }
+	virtual EAssetCommandResult ActivateAssets(const FAssetActivateArgs& ActivateArgs) const override;
+	virtual EAssetCommandResult OpenAssets(const FAssetOpenArgs& OpenArgs) const override;
+	// UAssetDefinition End
+
+private:
+	mutable bool bOpenFromContentBrowser = false;
+};
\ No newline at end of file
diff --git a/Engine/Source/Editor/WorldBookmark/Private/WorldBookmark/Browser/SWorldBookmarkBrowser.cpp b/Engine/Source/Editor/WorldBookmark/Private/WorldBookmark/Browser/SWorldBookmarkBrowser.cpp
new file mode 100644
--- /dev/null
+++ b/Engine/Source/Editor/WorldBookmark/Private/WorldBookmark/Browser/SWorldBookmarkBrowser.cpp
@@ -0,0 +1,797 @@
+// Copyright Epic Games, Inc. All Rights Reserved.
+#include "WorldBookmark/Browser/SWorldBookmarkBrowser.h"
+
+#include "AssetRegistry/IAssetRegistry.h"
+#include "AssetRegistry/AssetData.h"
+#include "AssetRegistry/AssetRegistryHelpers.h"
+#include "AssetToolsModule.h"
+#include "ContentBrowserDataSource.h"
+#include "ContentBrowserModule.h"
+#include "Editor.h"
+#include "LevelEditor.h"
+#include "Modules/ModuleManager.h"
+#include "PropertyEditorModule.h"
+#include "SSimpleButton.h"
+#include "Toolkits/GlobalEditorCommonCommands.h"
+#include "UObject/AssetRegistryTagsContext.h"
+#include "Widgets/Input/SSearchBox.h"
+#include "WorldBookmark/WorldBookmark.h"
+#include "WorldBookmark/WorldBookmarkFactory.h"
+
+#define LOCTEXT_NAMESPACE "WorldBookmarkBrowser"
+
+namespace WorldBookmarkBrowserColumn
+{
+	namespace Favorite
+	{
+		FName Id = TEXT("Favorite");
+		FText DisplayText = LOCTEXT("Favorite_Text", "");
+		FText ToolTipText = LOCTEXT("Favorite_Tooltip", "Favorite");
+	};
+
+	namespace RecentlyUsed
+	{
+		FName Id = TEXT("RecentlyUsed");
+		FText DisplayText = LOCTEXT("RecentlyUsed_Text", "");
+		FText ToolTipText = LOCTEXT("RecentlyUsed_Tooltip", "Recently Used");
+	};
+
+	namespace Label
+	{
+		FName Id = TEXT("Label");
+		FText DisplayText = LOCTEXT("Label_Text", "Name");
+		FText ToolTipText = LOCTEXT("Label_Tooltip", "World Bookmark name");
+	};
+}
+
+static UWorld* GetCurrentWorld()
+{
+	if (GEditor)
+	{
+		if (UWorld* EditorWorld = GEditor->GetEditorWorldContext().World())
+		{
+			return EditorWorld;
+		}
+	}
+	return nullptr;
+}
+
+SWorldBookmarkBrowser::SWorldBookmarkBrowser()
+{
+	// Listen for asset registry updates
+	IAssetRegistry& AssetRegistry = IAssetRegistry::GetChecked();
+	OnAssetAddedHandle = AssetRegistry.OnAssetsAdded().AddRaw(this, &SWorldBookmarkBrowser::OnAssetsAdded);
+	OnAssetRemovedHandle = AssetRegistry.OnAssetRemoved().AddRaw(this, &SWorldBookmarkBrowser::OnAssetRemoved);
+	OnAssetRenamedHandle = AssetRegistry.OnAssetRenamed().AddRaw(this, &SWorldBookmarkBrowser::OnAssetRenamed);
+	OnAssetUpdatedHandle = AssetRegistry.OnAssetUpdated().AddRaw(this, &SWorldBookmarkBrowser::OnAssetUpdated);
+	OnAssetUpdatedOnDiskHandle = AssetRegistry.OnAssetUpdatedOnDisk().AddRaw(this, &SWorldBookmarkBrowser::OnAssetUpdatedOnDisk);
+	
+	// Listen for when assets are loaded or changed
+	OnObjectPropertyChangedHandle = FCoreUObjectDelegates::OnObjectPropertyChanged.AddRaw(this, &SWorldBookmarkBrowser::OnObjectPropertyChanged);
+
+	// Listen for map change events
+	FLevelEditorModule& LevelEditorModule = FModuleManager::LoadModuleChecked<FLevelEditorModule>(TEXT("LevelEditor"));
+	OnMapChangedHandle = LevelEditorModule.OnMapChanged().AddRaw(this, &SWorldBookmarkBrowser::OnMapChanged);
+
+	ColumnsSortingParams[EColumnSortPriority::Primary] = { WorldBookmarkBrowserColumn::Favorite::Id, EColumnSortMode::Ascending };
+	ColumnsSortingParams[EColumnSortPriority::Secondary] = { WorldBookmarkBrowserColumn::Label::Id, EColumnSortMode::Ascending };
+}
+
+SWorldBookmarkBrowser::~SWorldBookmarkBrowser()
+{
+	IAssetRegistry& AssetRegistry = IAssetRegistry::GetChecked();
+	AssetRegistry.OnAssetAdded().Remove(OnAssetAddedHandle);
+	AssetRegistry.OnAssetRemoved().Remove(OnAssetRemovedHandle);
+	AssetRegistry.OnAssetRenamed().Remove(OnAssetRenamedHandle);
+	AssetRegistry.OnAssetUpdated().Remove(OnAssetUpdatedHandle);
+	AssetRegistry.OnAssetUpdatedOnDisk().Remove(OnAssetUpdatedOnDiskHandle);
+	
+	FCoreUObjectDelegates::OnObjectPropertyChanged.Remove(OnObjectPropertyChangedHandle);
+
+	FLevelEditorModule& LevelEditorModule = FModuleManager::LoadModuleChecked<FLevelEditorModule>(TEXT("LevelEditor"));
+	LevelEditorModule.OnMapChanged().Remove(OnMapChangedHandle);
+}
+
+static FSlateIcon GetFavoriteWorldBookmarkIcon(bool bIsFavorite)
+{
+	const FName FavoriteIconName = bIsFavorite ? "WorldPartition.WorldBookmark.IsFavorite" : "WorldPartition.WorldBookmark.IsNotFavorite";
+	return FSlateIcon(FAppStyle::Get().GetStyleSetName(), FavoriteIconName);
+}
+
+static FSlateIcon GetRecentlyUsedWorldBookmarkIcon()
+{
+	const FName RecentlyUsedIconName = "WorldPartition.WorldBookmark.RecentlyUsed";
+	return FSlateIcon(FAppStyle::Get().GetStyleSetName(), RecentlyUsedIconName);
+}
+
+void SWorldBookmarkBrowser::Construct(const FArguments& InArgs)
+{
+	RefreshItems();
+
+	// Create a property view
+	FPropertyEditorModule& EditModule = FModuleManager::Get().GetModuleChecked<FPropertyEditorModule>("PropertyEditor");
+	FDetailsViewArgs DetailsViewArgs;
+	DetailsViewArgs.bAllowSearch = false;
+	DetailsViewArgs.bHideSelectionTip = true;
+	PropertyView = EditModule.CreateDetailView(DetailsViewArgs);
+	PropertyView->SetObject(nullptr, true);
+
+	TSharedRef<SHeaderRow> HeaderRow = SNew(SHeaderRow)
+
+		// Favorite
+		+ SHeaderRow::Column(WorldBookmarkBrowserColumn::Favorite::Id)
+		.SortMode(this, &SWorldBookmarkBrowser::GetColumnSortMode, WorldBookmarkBrowserColumn::Favorite::Id)
+		.SortPriority(this, &SWorldBookmarkBrowser::GetColumnSortPriority, WorldBookmarkBrowserColumn::Favorite::Id)
+		.OnSort(this, &SWorldBookmarkBrowser::OnColumnSortChanged)
+		.FixedWidth(32)
+		.DefaultLabel(WorldBookmarkBrowserColumn::Favorite::DisplayText)
+		.DefaultTooltip(WorldBookmarkBrowserColumn::Favorite::ToolTipText)
+
+		// RecentlyUsed
+		+ SHeaderRow::Column(WorldBookmarkBrowserColumn::RecentlyUsed::Id)
+		.SortMode(this, &SWorldBookmarkBrowser::GetColumnSortMode, WorldBookmarkBrowserColumn::RecentlyUsed::Id)
+		.SortPriority(this, &SWorldBookmarkBrowser::GetColumnSortPriority, WorldBookmarkBrowserColumn::RecentlyUsed::Id)
+		.OnSort(this, &SWorldBookmarkBrowser::OnColumnSortChanged)
+		.FixedWidth(32)
+		.DefaultLabel(WorldBookmarkBrowserColumn::RecentlyUsed::DisplayText)
+		.DefaultTooltip(WorldBookmarkBrowserColumn::RecentlyUsed::ToolTipText)
+
+		// Label
+		+ SHeaderRow::Column(WorldBookmarkBrowserColumn::Label::Id)
+		.SortMode(this, &SWorldBookmarkBrowser::GetColumnSortMode, WorldBookmarkBrowserColumn::Label::Id)
+		.SortPriority(this, &SWorldBookmarkBrowser::GetColumnSortPriority, WorldBookmarkBrowserColumn::Label::Id)
+		.OnSort(this, &SWorldBookmarkBrowser::OnColumnSortChanged)
+		.DefaultLabel(WorldBookmarkBrowserColumn::Label::DisplayText)
+		.DefaultTooltip(WorldBookmarkBrowserColumn::Label::ToolTipText);
+
+	BookmarksView = SNew(STreeView<FWorldBookmarkTreeItemPtr>)
+		.TreeItemsSource(&WorldBookmarkTreeItems)
+		.SelectionMode(ESelectionMode::Single)
+		.OnGetChildren(this, &SWorldBookmarkBrowser::OnGetChildren)
+		.OnGenerateRow(this, &SWorldBookmarkBrowser::OnGenerateRow)
+		.OnSelectionChanged(this, &SWorldBookmarkBrowser::OnSelectedBookmarkChanged)
+		.OnMouseButtonDoubleClick(this, &SWorldBookmarkBrowser::OnBookmarkDoubleClicked)
+		.OnContextMenuOpening(this, &SWorldBookmarkBrowser::OnContextMenuOpening)
+		.HeaderRow
+		(
+			HeaderRow
+		);
+
+	TSharedRef<SSearchBox> SearchBoxRef = SNew(SSearchBox)
+		.OnTextChanged(this, &SWorldBookmarkBrowser::OnSearchBoxTextChanged);		
+
+	ChildSlot
+	[
+		SNew(SSplitter)
+		.Orientation(Orient_Vertical)
+		.Style(FAppStyle::Get(), "DetailsView.Splitter")
+
+		// Bookmark List
+		+ SSplitter::Slot()
+		[
+			SNew(SVerticalBox)
+
+			+ SVerticalBox::Slot()
+			.AutoHeight()
+			.VAlign(VAlign_Top)
+			.HAlign(HAlign_Fill)			
+			.Padding(0.0f)
+			[
+				SNew(SHorizontalBox)
+				
+				+ SHorizontalBox::Slot()
+				.HAlign(HAlign_Fill)
+				.FillWidth(1.0f)
+				[
+					SearchBoxRef
+				]
+
+				+ SHorizontalBox::Slot()
+				.AutoWidth()
+				.HAlign(HAlign_Center)
+				.VAlign(VAlign_Center)
+				[
+					SNew(SButton)
+					.ButtonStyle(FAppStyle::Get(), "SimpleButton")
+					.ToolTipText(LOCTEXT("CreateNewWorldBookmak", "Create a new world bookmark"))
+					.OnClicked(this, &SWorldBookmarkBrowser::OnCreateNewWorldBookmarkClicked)
+					[
+						SNew(SImage)
+						.ColorAndOpacity(FSlateColor::UseForeground())
+						.Image(FAppStyle::Get().GetBrush("WorldPartition.WorldBookmark.Create"))
+					]
+				]
+			]
+
+			+ SVerticalBox::Slot()
+			[
+				BookmarksView.ToSharedRef()
+			]
+		]
+
+		// Details Propery View
+		+ SSplitter::Slot()
+		[
+			SNew(SVerticalBox)
+
+			+ SVerticalBox::Slot()
+			.Padding(2, 4, 0, 0)
+			[
+				PropertyView.ToSharedRef()
+			]
+		]
+	];
+
+	Commands = MakeShareable(new FUICommandList());
+	Commands->MapAction(FGlobalEditorCommonCommands::Get().FindInContentBrowser, FUIAction(
+		FExecuteAction::CreateSP(this, &SWorldBookmarkBrowser::FindSelectedBookmarkInContentBrowser),
+		FCanExecuteAction::CreateSP(this, &SWorldBookmarkBrowser::IsValidBookmarkSelected)
+	));
+}
+
+void SWorldBookmarkBrowser::RefreshItems()
+{
+	IAssetRegistry& AssetRegistry = IAssetRegistry::GetChecked();
+	TArray<FAssetData> AssetsData;
+	const bool bSearchSubClasses = true;
+	AssetRegistry.GetAssetsByClass(UWorldBookmark::StaticClass()->GetClassPathName(), AssetsData, bSearchSubClasses);
+
+	UWorldBookmark* PreviousSelection = GetSelectedBookmark();
+
+	WorldBookmarkTreeItems.Reset();
+
+	for (const FAssetData& AssetData : AssetsData)
+	{
+		static const FName NAME_BookmarkWorld(TEXT("WorldName"));
+
+		FString WorldAssetName;
+		if (!AssetData.GetTagValue(NAME_BookmarkWorld, WorldAssetName))
+		{
+			continue;
+		}
+
+		FSoftObjectPath WorldAssetPath = FSoftObjectPath(WorldAssetName);
+		UAssetRegistryHelpers::FixupRedirectedAssetPath(WorldAssetPath);
+
+		FSoftObjectPath CurrentWorldPath = FSoftObjectPath(GetCurrentWorld());
+		if (CurrentWorldPath != WorldAssetPath)
+		{
+			continue;
+		}
+
+		if (!CurrentSearchString.IsEmpty() && !AssetData.AssetName.ToString().Contains(CurrentSearchString.ToString()))
+		{
+			continue;
+		}
+
+		FWorldBookmarkTreeItemPtr& Item = WorldBookmarkTreeItems.Emplace_GetRef(MakeShared<FWorldBookmarkTreeItem>());
+		Item->BookmarkAsset = AssetData;
+	}
+
+	SortItems();	
+
+	// Restore selected bookmark after refresh
+	SetSelectedBookmark(PreviousSelection);
+}
+
+void SWorldBookmarkBrowser::SortItems()
+{
+	WorldBookmarkTreeItems.Sort([this](const FWorldBookmarkTreeItemPtr& A, const FWorldBookmarkTreeItemPtr& B)
+	{
+		auto CompareItems = [&](const TPair<FName, EColumnSortMode::Type>& SortOrder) -> int32
+		{
+			auto ApplySortMode = [&](int32 SortResult) { return SortOrder.Value == EColumnSortMode::Ascending ? SortResult : -SortResult; };
+			
+			if (!A.IsValid() || !B.IsValid())
+			{
+				return A.IsValid() ? -1 : 0;
+			}
+
+			UWorldBookmark* WorldBookmarkA = Cast<UWorldBookmark>(A->BookmarkAsset.GetAsset());
+			UWorldBookmark* WorldBookmarkB = Cast<UWorldBookmark>(B->BookmarkAsset.GetAsset());
+			
+			if (!WorldBookmarkA || !WorldBookmarkB)
+			{
+				return WorldBookmarkA ? -1 : 0;
+			}
+
+			if (SortOrder.Key == WorldBookmarkBrowserColumn::Favorite::Id)
+			{
+				return ApplySortMode(WorldBookmarkA->GetIsUserFavorite() == WorldBookmarkB->GetIsUserFavorite() ? 0 : WorldBookmarkA->GetIsUserFavorite() ? -1 : 1);
+			}
+			else if (SortOrder.Key == WorldBookmarkBrowserColumn::RecentlyUsed::Id)
+			{
+				return ApplySortMode(WorldBookmarkA->GetUserLastLoadedTimeStampUTC() == WorldBookmarkB->GetUserLastLoadedTimeStampUTC() ? 0 : WorldBookmarkA->GetUserLastLoadedTimeStampUTC() < WorldBookmarkB->GetUserLastLoadedTimeStampUTC() ? -1 : 1);
+			}
+			else if (SortOrder.Key == WorldBookmarkBrowserColumn::Label::Id)
+			{
+				return ApplySortMode(WorldBookmarkA->GetName() == WorldBookmarkB->GetName() ? 0 : WorldBookmarkA->GetName() < WorldBookmarkB->GetName() ? -1 : 1);
+			}
+		
+			checkf(false, TEXT("Unsupported sort order: %s"), *SortOrder.Key.ToString());
+			return 0;
+		};
+
+		// Primary Sort
+		int32 PrimarySortOrder = CompareItems(ColumnsSortingParams[EColumnSortPriority::Primary]);
+		if (PrimarySortOrder != 0)
+		{
+			return PrimarySortOrder < 0;
+		}
+
+		// Secondary Sort
+		int32 SecondarySortOrder = CompareItems(ColumnsSortingParams[EColumnSortPriority::Secondary]);
+		return SecondarySortOrder < 0;
+	});
+
+	if (BookmarksView)
+	{
+		BookmarksView->RequestListRefresh();
+	}
+}
+
+void FWorldBookmarkTableRow::Construct(const FArguments& InArgs, const TSharedRef<STableViewBase>& OwnerTable, FWorldBookmarkTreeItemRef InItem)
+{
+	WeakWorldBookmarkTreeItem = InItem;	
+	SMultiColumnTableRow<FWorldBookmarkTreeItemPtr>::Construct(InArgs, OwnerTable);
+}
+
+TSharedRef<SWidget> FWorldBookmarkTableRow::GenerateWidgetForColumn(const FName& ColumnName)
+{
+	if (WeakWorldBookmarkTreeItem.IsValid())
+	{
+		if (ColumnName == WorldBookmarkBrowserColumn::Favorite::Id)
+		{
+			if (UWorldBookmark* WorldBookmark = Cast<UWorldBookmark>(WeakWorldBookmarkTreeItem.Pin()->BookmarkAsset.GetAsset()))
+			{
+				if (WorldBookmark->GetIsUserFavorite())
+				{
+					return SNew(SBox)
+						.Padding(2.0f)
+						.HAlign(HAlign_Left)
+						.VAlign(VAlign_Center)
+						[
+							SNew(SImage)
+								.DesiredSizeOverride(FVector2D(24, 24))
+								.ColorAndOpacity(FSlateColor::UseForeground())
+								.Image(GetFavoriteWorldBookmarkIcon(true).GetIcon())
+						];
+				}
+			}
+		}
+		else if (ColumnName == WorldBookmarkBrowserColumn::RecentlyUsed::Id)
+		{
+			if (UWorldBookmark* WorldBookmark = Cast<UWorldBookmark>(WeakWorldBookmarkTreeItem.Pin()->BookmarkAsset.GetAsset()))
+			{
+				FTimespan LastLoaded = FDateTime::UtcNow() - WorldBookmark->GetUserLastLoadedTimeStampUTC();
+				if (LastLoaded < FTimespan::FromDays(14))
+				{
+					return SNew(SBox)
+						.Padding(2.0f)
+						.HAlign(HAlign_Left)
+						.VAlign(VAlign_Center)
+						[
+							SNew(SImage)
+								.DesiredSizeOverride(FVector2D(24, 24))
+								.ColorAndOpacity(FSlateColor::UseForeground())
+								.Image(GetRecentlyUsedWorldBookmarkIcon().GetIcon())
+						];
+				}
+			}
+		}
+		else if (ColumnName == WorldBookmarkBrowserColumn::Label::Id)
+		{
+			return SNew(SVerticalBox)
+			
+			+ SVerticalBox::Slot()
+			.Padding(10, 5, 0, 0)
+			.AutoHeight()
+			[
+				SNew(STextBlock)
+					.Text(FText::FromName(WeakWorldBookmarkTreeItem.Pin()->BookmarkAsset.AssetName))
+			]
+
+			+ SVerticalBox::Slot()
+			.Padding(10, 0, 0, 5)
+			.AutoHeight()
+			[
+				SNew(STextBlock)
+					.ColorAndOpacity(FSlateColor::UseSubduedForeground())
+					.Text(this, &FWorldBookmarkTableRow::GetLastAccessedText)
+			];
+		}
+	}
+
+	return SNullWidget::NullWidget;
+}
+
+FText FWorldBookmarkTableRow::GetLastAccessedText() const
+{
+	if (WeakWorldBookmarkTreeItem.IsValid() && WeakWorldBookmarkTreeItem.Pin()->BookmarkAsset.IsValid())
+	{
+		if (UWorldBookmark* WorldBookmark = Cast<UWorldBookmark>(WeakWorldBookmarkTreeItem.Pin()->BookmarkAsset.GetAsset()))
+		{
+			FDateTime LastLoadedUTC = WorldBookmark->GetUserLastLoadedTimeStampUTC();
+
+			if (LastLoadedUTC != FDateTime::MinValue())
+			{
+				FTimespan UTCOffset = FDateTime::UtcNow() - FDateTime::Now();
+				uint64 AdjustedTicks = LastLoadedUTC.GetTicks() - UTCOffset.GetTicks();
+				FDateTime LastLoadedLocal = FDateTime(AdjustedTicks);
+
+				return FText::Format(LOCTEXT("LastAccessedDateTime", "Last accessed {0}"), FText::AsDateTime(LastLoadedUTC, TEXT("%Y-%m-%d %H:%M")));
+			}
+		}
+	}
+
+	return FText();
+}
+
+void SWorldBookmarkBrowser::OnGetChildren(FWorldBookmarkTreeItemPtr InItem, TArray<FWorldBookmarkTreeItemPtr>& OutChildren)
+{
+}
+
+TSharedRef<ITableRow> SWorldBookmarkBrowser::OnGenerateRow(FWorldBookmarkTreeItemPtr InTreeItem, const TSharedRef<STableViewBase>& OwnerTable)
+{
+	return SNew(FWorldBookmarkTableRow, OwnerTable, InTreeItem.ToSharedRef());
+}
+
+void SWorldBookmarkBrowser::OnSelectedBookmarkChanged(FWorldBookmarkTreeItemPtr InTreeItem, ESelectInfo::Type SelectionType) const
+{
+	UWorldBookmark* SelectedWorldBookmark = nullptr;
+
+	if (InTreeItem.IsValid() && InTreeItem->BookmarkAsset.IsValid())
+	{
+		if (UWorldBookmark* WorldBookmark = Cast<UWorldBookmark>(InTreeItem->BookmarkAsset.GetAsset()))
+		{
+			SelectedWorldBookmark = WorldBookmark;
+		}
+	}
+
+	PropertyView->SetObject(SelectedWorldBookmark);
+}
+
+void SWorldBookmarkBrowser::OnBookmarkDoubleClicked(FWorldBookmarkTreeItemPtr InTreeItem)
+{
+	if (InTreeItem.IsValid() && InTreeItem->BookmarkAsset.IsValid())
+	{
+		if (UWorldBookmark* WorldBookmark = Cast<UWorldBookmark>(InTreeItem->BookmarkAsset.GetAsset()))
+		{
+			WorldBookmark->Load();
+
+			// Loading an item will update it's last load time - refresh the UI
+			RefreshItems();
+		}
+	}
+}
+
+FReply SWorldBookmarkBrowser::OnCreateNewWorldBookmarkClicked()
+{
+	FAssetToolsModule& AssetToolsModule = FAssetToolsModule::GetModule();
+	UWorldBookmarkFactory* Factory = NewObject<UWorldBookmarkFactory>();
+
+	// Create a new bookmark asset
+	UObject* NewAsset = AssetToolsModule.Get().CreateAssetWithDialog(UWorldBookmark::StaticClass(), Factory);
+	if (UWorldBookmark* NewWorldBookmark = Cast<UWorldBookmark>(NewAsset))
+	{
+		// Update it with the current state of the editor
+		NewWorldBookmark->Update();
+
+		// Show it in the content browser
+		GEditor->SyncBrowserToObject(NewWorldBookmark);
+
+		// Refresh the bookmark browser to ensure the new bookmark is shown
+		RefreshItems();
+
+		// Make sure the new bookmark is selected
+		SetSelectedBookmark(NewWorldBookmark);
+	}
+
+	return FReply::Handled();
+}
+
+void SWorldBookmarkBrowser::OnSearchBoxTextChanged(const FText& InSearchText)
+{
+	if (!CurrentSearchString.EqualTo(InSearchText))
+	{
+		CurrentSearchString = InSearchText;
+		RefreshItems();
+	}
+}
+
+EColumnSortMode::Type SWorldBookmarkBrowser::GetColumnSortMode(const FName InColumnId) const
+{
+	if (ColumnsSortingParams[EColumnSortPriority::Primary].Key == InColumnId)
+	{
+		return ColumnsSortingParams[EColumnSortPriority::Primary].Value;
+	}
+
+	if (ColumnsSortingParams[EColumnSortPriority::Secondary].Key == InColumnId)
+	{
+		return ColumnsSortingParams[EColumnSortPriority::Secondary].Value;
+	}
+
+	return EColumnSortMode::None;
+}
+
+EColumnSortPriority::Type SWorldBookmarkBrowser::GetColumnSortPriority(const FName InColumnId) const
+{
+	if (ColumnsSortingParams[EColumnSortPriority::Primary].Key == InColumnId)
+	{
+		return EColumnSortPriority::Primary;
+	}
+
+	if (ColumnsSortingParams[EColumnSortPriority::Secondary].Key == InColumnId)
+	{
+		return EColumnSortPriority::Secondary;
+	}
+
+	return EColumnSortPriority::None;
+}
+
+void SWorldBookmarkBrowser::OnColumnSortChanged(const EColumnSortPriority::Type InSortPriority, const FName& InColumnId, const EColumnSortMode::Type InSortMode)
+{
+	ColumnsSortingParams[InSortPriority].Key = InColumnId;
+	ColumnsSortingParams[InSortPriority].Value = InSortMode;
+
+	SortItems();
+}
+
+TSharedPtr<SWidget> SWorldBookmarkBrowser::OnContextMenuOpening()
+{
+	FMenuBuilder MenuBuilder(true, Commands);
+
+	// World Bookmark
+	MenuBuilder.BeginSection(TEXT("WorldBookmarkMenuSection"), LOCTEXT("WorldBookmarkMenuSection_Label", "World Bookmark"));
+	{
+		// Load
+		MenuBuilder.AddMenuEntry(
+			LOCTEXT("LoadBookmark", "Load Bookmark"),
+			LOCTEXT("LoadBookmarkToolTip", "Load this bookmark and apply it's state to the world."),
+			FSlateIcon(FAppStyle::Get().GetStyleSetName(), "WorldPartition.WorldBookmark.GoTo"),
+			FUIAction(
+				FExecuteAction::CreateSP(this, &SWorldBookmarkBrowser::LoadSelectedBookmark),
+				FCanExecuteAction::CreateSP(this, &SWorldBookmarkBrowser::IsValidBookmarkSelected)
+			)
+		);
+
+		// Update
+		MenuBuilder.AddMenuEntry(
+			LOCTEXT("UpdateBookmark", "Update Bookmark"),
+			LOCTEXT("UpdateBookmarkToolTip", "Update this bookmark using the current state of the world."),
+			FSlateIcon(FAppStyle::Get().GetStyleSetName(), "WorldPartition.WorldBookmark.Update"),
+			FUIAction(
+				FExecuteAction::CreateSP(this, &SWorldBookmarkBrowser::UpdateSelectedBookmark),
+				FCanExecuteAction::CreateSP(this, &SWorldBookmarkBrowser::IsValidBookmarkSelected)
+			)
+		);
+
+		MenuBuilder.AddMenuEntry(
+			LOCTEXT("AddToFavorite", "Add to Favorite"),
+			LOCTEXT("AddToFavorite_ToolTip", "Add this bookmark to your favorites."),
+			GetFavoriteWorldBookmarkIcon(false),
+			FUIAction(
+				FExecuteAction::CreateSP(this, &SWorldBookmarkBrowser::AddSelectedBookmarkToFavorites),
+				FCanExecuteAction::CreateSP(this, &SWorldBookmarkBrowser::IsValidBookmarkSelected),
+				FIsActionChecked(),
+				FIsActionButtonVisible::CreateSP(this, &SWorldBookmarkBrowser::IsSelectedBookmarkNotFavorite)
+			)
+		);
+
+		MenuBuilder.AddMenuEntry(
+			LOCTEXT("RemoveFromFavorite", "Remove from Favorite"),
+			LOCTEXT("RemoveFromFavorite_ToolTip", "Remove this bookmark from your favorites."),
+			GetFavoriteWorldBookmarkIcon(true),
+			FUIAction(
+				FExecuteAction::CreateSP(this, &SWorldBookmarkBrowser::RemoveSelectedBookmarkFromFavorites),
+				FCanExecuteAction::CreateSP(this, &SWorldBookmarkBrowser::IsValidBookmarkSelected),
+				FIsActionChecked(),
+				FIsActionButtonVisible::CreateSP(this, &SWorldBookmarkBrowser::IsSelectedBookmarkFavorite)
+			)
+		);
+	}
+
+	MenuBuilder.EndSection();
+
+	// Asset Options
+	MenuBuilder.BeginSection(TEXT("AssetOptionsSection"), LOCTEXT("AssetOptionsText", "Asset Options"));
+	{
+		// Find In Content Browser
+		MenuBuilder.AddMenuEntry(FGlobalEditorCommonCommands::Get().FindInContentBrowser);
+	}
+	MenuBuilder.EndSection();
+
+	return MenuBuilder.MakeWidget();
+}
+
+bool IsWorldBookmarkAsset(const FAssetData& InAssetData)
+{
+	return InAssetData.AssetClassPath == UWorldBookmark::StaticClass()->GetClassPathName();
+}
+
+bool IsWorldBookmarkAsset(UObject* InObject)
+{
+	return InObject->IsA<UWorldBookmark>();
+}
+
+void SWorldBookmarkBrowser::OnAssetsAdded(TConstArrayView<FAssetData> InAssets)
+{
+	for (const FAssetData& InAssetData : InAssets)
+	{
+		if (IsWorldBookmarkAsset(InAssetData))
+		{
+			RefreshItems();
+		}
+	}
+}
+
+void SWorldBookmarkBrowser::OnAssetRemoved(const FAssetData& InAssetData)
+{
+	if (IsWorldBookmarkAsset(InAssetData))
+	{
+		RefreshItems();
+	}
+}
+
+void SWorldBookmarkBrowser::OnAssetRenamed(const FAssetData& InAssetData, const FString& InOldObjectPath)
+{
+	if (IsWorldBookmarkAsset(InAssetData))
+	{
+		RefreshItems();
+	}
+}
+
+void SWorldBookmarkBrowser::OnAssetUpdated(const FAssetData& InAssetData)
+{
+	if (IsWorldBookmarkAsset(InAssetData))
+	{
+		RefreshItems();
+	}
+}
+
+void SWorldBookmarkBrowser::OnAssetUpdatedOnDisk(const FAssetData& InAssetData)
+{
+	if (IsWorldBookmarkAsset(InAssetData))
+	{
+		RefreshItems();
+	}
+}
+
+void SWorldBookmarkBrowser::OnObjectPropertyChanged(UObject* InObject, FPropertyChangedEvent& InPropertyChangedEvent)
+{
+	if (IsWorldBookmarkAsset(InObject))
+	{
+		RefreshItems();
+	}
+}
+
+void SWorldBookmarkBrowser::OnObjectPreSave(UObject* InObject, FObjectPreSaveContext InObjectPreSaveContext)
+{
+	if (IsWorldBookmarkAsset(InObject))
+	{
+		RefreshItems();
+	}
+}
+
+void SWorldBookmarkBrowser::OnMapChanged(UWorld* World, EMapChangeType MapChangeType)
+{
+	RefreshItems();
+}
+
+void SWorldBookmarkBrowser::LoadSelectedBookmark()
+{
+	if (UWorldBookmark* SelectedBookmark = GetSelectedBookmark())
+	{
+		SelectedBookmark->Load();
+	}
+}
+
+void SWorldBookmarkBrowser::UpdateSelectedBookmark()
+{
+	if (UWorldBookmark* SelectedBookmark = GetSelectedBookmark())
+	{
+		SelectedBookmark->Update();
+		PropertyView->ForceRefresh();
+	}
+}
+
+void SWorldBookmarkBrowser::FindSelectedBookmarkInContentBrowser()
+{
+	if (UWorldBookmark* SelectedBookmark = GetSelectedBookmark())
+	{
+		FContentBrowserModule& ContentBrowserModule = FModuleManager::Get().LoadModuleChecked<FContentBrowserModule>("ContentBrowser");
+		ContentBrowserModule.Get().SyncBrowserToAssets(TArray<UObject*>({ SelectedBookmark }));
+	}
+}
+
+void SWorldBookmarkBrowser::AddSelectedBookmarkToFavorites()
+{
+	if (UWorldBookmark* SelectedBookmark = GetSelectedBookmark())
+	{
+		return SelectedBookmark->SetIsUserFavorite(true);
+	}
+}
+
+void SWorldBookmarkBrowser::RemoveSelectedBookmarkFromFavorites()
+{
+	if (UWorldBookmark* SelectedBookmark = GetSelectedBookmark())
+	{
+		return SelectedBookmark->SetIsUserFavorite(false);
+	}
+}
+
+bool SWorldBookmarkBrowser::IsSelectedBookmarkFavorite() const
+{
+	if (UWorldBookmark* SelectedBookmark = GetSelectedBookmark())
+	{
+		return SelectedBookmark->GetIsUserFavorite();
+	}
+
+	return false;
+}
+
+bool SWorldBookmarkBrowser::IsSelectedBookmarkNotFavorite() const
+{
+	if (UWorldBookmark* SelectedBookmark = GetSelectedBookmark())
+	{
+		return !SelectedBookmark->GetIsUserFavorite();
+	}
+
+	return true;
+}
+
+bool SWorldBookmarkBrowser::IsValidBookmarkSelected() const
+{
+	return GetSelectedBookmark() != nullptr;
+}
+
+UWorldBookmark* SWorldBookmarkBrowser::GetSelectedBookmark() const
+{
+	if (!BookmarksView)
+	{
+		return nullptr;
+	}
+
+	const TArray<FWorldBookmarkTreeItemPtr> Selection = BookmarksView->GetSelectedItems();
+	if (Selection.Num() != 1)
+	{
+		return nullptr;
+	}
+
+	FWorldBookmarkTreeItemPtr SelectedItem = Selection[0];
+	if (!SelectedItem.IsValid() || !SelectedItem->BookmarkAsset.IsValid())
+	{
+		return nullptr;
+	}
+
+	return Cast<UWorldBookmark>(SelectedItem->BookmarkAsset.GetAsset());
+}
+
+void SWorldBookmarkBrowser::SetSelectedBookmark(UWorldBookmark* InBookmark)
+{
+	if (!BookmarksView)
+	{
+		return;
+	}
+
+	if (InBookmark)
+	{
+		FWorldBookmarkTreeItemPtr* BookmarkTreeItem = WorldBookmarkTreeItems.FindByPredicate([InBookmark](FWorldBookmarkTreeItemPtr InItem)
+		{
+			return InItem->BookmarkAsset == InBookmark;
+		});
+
+		if (BookmarkTreeItem)
+		{
+			BookmarksView->SetSelection(*BookmarkTreeItem);
+			BookmarksView->RequestScrollIntoView(*BookmarkTreeItem);
+		}
+	}
+	else
+	{
+		BookmarksView->ClearSelection();
+	}
+}
+
+#undef LOCTEXT_NAMESPACE
+
diff --git a/Engine/Source/Editor/WorldBookmark/Private/WorldBookmark/Browser/SWorldBookmarkBrowser.h b/Engine/Source/Editor/WorldBookmark/Private/WorldBookmark/Browser/SWorldBookmarkBrowser.h
new file mode 100644
--- /dev/null
+++ b/Engine/Source/Editor/WorldBookmark/Private/WorldBookmark/Browser/SWorldBookmarkBrowser.h
@@ -0,0 +1,112 @@
+// Copyright Epic Games, Inc. All Rights Reserved.
+#pragma once
+
+#include "Misc/TextFilter.h"
+#include "Templates/SharedPointer.h"
+#include "Widgets/Views/SHeaderRow.h"	
+#include "Widgets/Views/STreeView.h"
+#include "ContentBrowserItem.h"
+#include "AssetTextFilter.h"
+#include "UObject/ObjectSaveContext.h"
+#include "UnrealEdMisc.h"
+
+class UWorldBookmark;
+class IDetailsView;
+class SAssetView;
+
+struct FWorldBookmarkTreeItem : TSharedFromThis<FWorldBookmarkTreeItem>
+{
+public:
+	FAssetData BookmarkAsset;
+};
+
+typedef TSharedPtr<FWorldBookmarkTreeItem> FWorldBookmarkTreeItemPtr;
+typedef TSharedRef<FWorldBookmarkTreeItem> FWorldBookmarkTreeItemRef;
+
+class FWorldBookmarkTableRow : public SMultiColumnTableRow<FWorldBookmarkTreeItemPtr>
+{
+public:
+	void Construct(const FArguments& InArgs, const TSharedRef<STableViewBase>& OwnerTable, FWorldBookmarkTreeItemRef InItem);
+
+private:
+	virtual TSharedRef<SWidget> GenerateWidgetForColumn(const FName& ColumnName) override;
+
+	FText GetLastAccessedText() const;
+
+private:
+	TWeakPtr<FWorldBookmarkTreeItem> WeakWorldBookmarkTreeItem;
+};
+
+class SWorldBookmarkBrowser : public SCompoundWidget
+{
+public:
+	SLATE_BEGIN_ARGS(SWorldBookmarkBrowser)
+		{}
+	SLATE_END_ARGS()
+
+	SWorldBookmarkBrowser();
+	virtual ~SWorldBookmarkBrowser();
+
+	void Construct(const FArguments& InArgs);
+
+private:
+	void OnGetChildren(FWorldBookmarkTreeItemPtr InItem, TArray<FWorldBookmarkTreeItemPtr>& OutChildren);
+	TSharedRef<ITableRow> OnGenerateRow(FWorldBookmarkTreeItemPtr InTreeItem, const TSharedRef<STableViewBase>& OwnerTable);
+	void OnSelectedBookmarkChanged(FWorldBookmarkTreeItemPtr InTreeItem, ESelectInfo::Type SelectionType) const;
+	void OnBookmarkDoubleClicked(FWorldBookmarkTreeItemPtr InTreeItem);
+	FReply OnCreateNewWorldBookmarkClicked();
+	void OnSearchBoxTextChanged(const FText& InSearchText);
+	void OnColumnSortChanged(const EColumnSortPriority::Type SortPriority, const FName& ColumnId, const EColumnSortMode::Type InSortMode);
+	TSharedPtr<SWidget> OnContextMenuOpening();
+
+	void RefreshItems();
+	void SortItems();
+
+	EColumnSortMode::Type GetColumnSortMode(const FName ColumnId) const;
+	EColumnSortPriority::Type GetColumnSortPriority(const FName InColumnId) const;
+	TArray<FAssetViewCustomColumn> GetCustomColumns() const;
+
+	void OnAssetsAdded(TConstArrayView<FAssetData> InAssets);
+	void OnAssetRemoved(const FAssetData& InAssetData);
+	void OnAssetRenamed(const FAssetData& InAssetData, const FString& InOldObjectPath);
+	void OnAssetUpdated(const FAssetData& InAssetData);
+	void OnAssetUpdatedOnDisk(const FAssetData& InAssetData);
+	void OnObjectPropertyChanged(UObject* InObject, FPropertyChangedEvent& InPropertyChangedEvent);
+	void OnObjectPreSave(UObject* InObject, FObjectPreSaveContext InObjectPreSaveContext);
+
+	void OnMapChanged(UWorld* World, EMapChangeType MapChangeType);
+
+	UWorldBookmark* GetSelectedBookmark() const;
+	void SetSelectedBookmark(UWorldBookmark* InBookmark);
+
+	void LoadSelectedBookmark();
+	void UpdateSelectedBookmark();
+	void FindSelectedBookmarkInContentBrowser();
+	bool IsValidBookmarkSelected() const;
+
+	bool IsSelectedBookmarkFavorite() const;
+	bool IsSelectedBookmarkNotFavorite() const;
+	void AddSelectedBookmarkToFavorites();
+	void RemoveSelectedBookmarkFromFavorites();	
+
+private:
+	FDelegateHandle OnAssetAddedHandle;
+	FDelegateHandle OnAssetRemovedHandle;
+	FDelegateHandle OnAssetRenamedHandle;
+	FDelegateHandle OnAssetUpdatedHandle;
+	FDelegateHandle OnAssetUpdatedOnDiskHandle;
+	FDelegateHandle OnObjectPropertyChangedHandle;
+	FDelegateHandle OnMapChangedHandle;
+
+	TArray<FWorldBookmarkTreeItemPtr>	WorldBookmarkTreeItems;
+
+	TSharedPtr<STreeView<FWorldBookmarkTreeItemPtr>> BookmarksView;
+
+	TSharedPtr<IDetailsView> PropertyView;
+
+	TSharedPtr<FUICommandList> Commands;
+
+	TPair<FName, EColumnSortMode::Type> ColumnsSortingParams[EColumnSortPriority::Max];
+	
+	FText CurrentSearchString;
+};
\ No newline at end of file
diff --git a/Engine/Source/Editor/WorldBookmark/Private/WorldBookmark/WorldBookmark.cpp b/Engine/Source/Editor/WorldBookmark/Private/WorldBookmark/WorldBookmark.cpp
new file mode 100644
--- /dev/null
+++ b/Engine/Source/Editor/WorldBookmark/Private/WorldBookmark/WorldBookmark.cpp
@@ -0,0 +1,134 @@
+// Copyright Epic Games, Inc. All Rights Reserved.
+
+#include "WorldBookmark/WorldBookmark.h"
+#include "EditorState/EditorStateSubsystem.h"
+#include "EditorState/WorldEditorState.h"
+
+#include "AssetRegistry/AssetRegistryHelpers.h"
+#include "UObject/AssetRegistryTagsContext.h"
+#include "FileHelpers.h"
+#include "Misc/MessageDialog.h"
+#include "EngineUtils.h"
+#include "ScopedTransaction.h"
+
+#include UE_INLINE_GENERATED_CPP_BY_NAME(WorldBookmark)
+
+#define LOCTEXT_NAMESPACE "WorldBookmark"
+
+UWorldBookmark::UWorldBookmark(const FObjectInitializer& ObjectInitializer)
+	: Super(ObjectInitializer)
+{
+}
+
+void UWorldBookmark::PostInitProperties()
+{
+	Super::PostInitProperties();
+
+	if (!HasAnyFlags(RF_ClassDefaultObject | RF_NeedLoad))
+	{
+		BookmarkGuid = FGuid::NewGuid();
+		BookmarkAssetPath = GetPathName();
+	}
+}
+
+void UWorldBookmark::PostDuplicate(bool bDuplicateForPIE)
+{
+	Super::PostDuplicate(bDuplicateForPIE);
+
+	BookmarkGuid = FGuid::NewGuid();
+	BookmarkAssetPath = GetPathName();
+}
+
+void UWorldBookmark::PostLoad()
+{
+	Super::PostLoad();
+
+	check(BookmarkGuid.IsValid());
+	BookmarkAssetPath = GetPathName();
+}
+
+void UWorldBookmark::PreSave(FObjectPreSaveContext ObjectSaveContext)
+{
+	Super::PreSave(ObjectSaveContext);
+	SaveConfig();
+}
+
+void UWorldBookmark::OverridePerObjectConfigSection(FString& SectionName)
+{
+	SectionName = FString::Printf(TEXT("WorldBookmark %s"), *BookmarkGuid.ToString());
+}
+
+void UWorldBookmark::GetAssetRegistryTags(FAssetRegistryTagsContext Context) const
+{
+	Super::GetAssetRegistryTags(Context);
+
+	if (EditorState.HasState<UWorldEditorState>())
+	{
+		TSoftObjectPtr<UWorld> World = EditorState.GetStateChecked<UWorldEditorState>().GetStateWorld();
+		if (!World.IsNull())
+		{
+			FSoftObjectPath WorldPath = FSoftObjectPath(World.ToString());
+			UAssetRegistryHelpers::FixupRedirectedAssetPath(WorldPath);
+
+			static const FName NAME_BookmarkWorld = "WorldName";
+			Context.AddTag(FAssetRegistryTag(NAME_BookmarkWorld, WorldPath.ToString(), FAssetRegistryTag::TT_Hidden));
+		}
+	}
+}
+
+void UWorldBookmark::Load()
+{
+	UEditorStateSubsystem::Get()->RestoreEditorState(EditorState);
+
+	SetUserLastLoadedTimeStampUTC(FDateTime::UtcNow());
+	SaveConfig();
+}
+
+void UWorldBookmark::Update()
+{
+	UWorld* CurrentWorld = GWorld;
+	if (CurrentWorld != nullptr && !CurrentWorld->IsGameWorld())
+	{
+		FScopedTransaction ScopedTransaction(LOCTEXT("UpdateWorldBookmarkTransaction", "Update World Bookmark"));
+
+		UWorld* BookmarkWorld = EditorState.HasState<UWorldEditorState>() ? EditorState.GetState<UWorldEditorState>()->GetStateWorld().Get() : nullptr;
+		if (BookmarkWorld != nullptr && BookmarkWorld != CurrentWorld)
+		{
+			FMessageDialog::Open(EAppMsgType::Ok, LOCTEXT("CannotUpdateWorldBookmark", "The bookmark you are trying to update does not belong to the current world."));
+			return;
+		}
+
+		Modify();
+
+		UEditorStateSubsystem::Get()->CaptureEditorState(EditorState, this);
+
+		PostEditChange();
+	}
+}
+
+bool UWorldBookmark::GetIsUserFavorite() const
+{
+	return bFavorite;
+}
+
+void UWorldBookmark::SetIsUserFavorite(bool bIsUserFavorite)
+{
+	if (bIsUserFavorite != bFavorite)
+	{
+		bFavorite = bIsUserFavorite;
+		SaveConfig();
+	}
+}
+
+FDateTime UWorldBookmark::GetUserLastLoadedTimeStampUTC() const
+{
+	return LastLoadedTimeStampUTC;
+}
+
+void UWorldBookmark::SetUserLastLoadedTimeStampUTC(const FDateTime& InLastLoadedTimeStampUTC)
+{
+	LastLoadedTimeStampUTC = InLastLoadedTimeStampUTC;
+	SaveConfig();
+}
+
+#undef LOCTEXT_NAMESPACE
diff --git a/Engine/Source/Editor/WorldBookmark/Private/WorldBookmark/WorldBookmarkDetailsCustomization.cpp b/Engine/Source/Editor/WorldBookmark/Private/WorldBookmark/WorldBookmarkDetailsCustomization.cpp
new file mode 100644
--- /dev/null
+++ b/Engine/Source/Editor/WorldBookmark/Private/WorldBookmark/WorldBookmarkDetailsCustomization.cpp
@@ -0,0 +1,75 @@
+// Copyright Epic Games, Inc. All Rights Reserved.
+
+#include "WorldBookmark/WorldBookmarkDetailsCustomization.h"
+#include "WorldBookmark/WorldBookmark.h"
+#include "DetailLayoutBuilder.h"
+#include "DetailCategoryBuilder.h"
+#include "Editor.h"
+#include "EditorState/WorldEditorState.h"
+
+FWorldBookmarkDetailsCustomization::FWorldBookmarkDetailsCustomization()
+	: CachedDetailBuilder()
+{
+	GEditor->RegisterForUndo(this);
+}
+
+FWorldBookmarkDetailsCustomization::~FWorldBookmarkDetailsCustomization()
+{
+	GEditor->UnregisterForUndo(this);
+}
+
+TSharedRef<IDetailCustomization> FWorldBookmarkDetailsCustomization::MakeInstance()
+{
+	return MakeShareable(new FWorldBookmarkDetailsCustomization);
+}
+
+void FWorldBookmarkDetailsCustomization::CustomizeDetails(IDetailLayoutBuilder& DetailBuilder)
+{
+	CachedDetailBuilder = &DetailBuilder;
+
+	TArray<TWeakObjectPtr<UObject>> ObjectsBeingCustomized;
+	DetailBuilder.GetObjectsBeingCustomized(ObjectsBeingCustomized);
+	if (ObjectsBeingCustomized.Num() > 1)
+	{
+		return;
+	}
+
+	UWorldBookmark* WorldBookmark = Cast<UWorldBookmark>(ObjectsBeingCustomized[0].Get());
+	if (WorldBookmark == nullptr)
+	{
+		return;
+	}
+
+	// Sort categories alphabetically
+	TArray<UEditorState*> SortedEditorStateObjects = WorldBookmark->EditorState.GetStates();
+	SortedEditorStateObjects.Sort([](UEditorState& A, UEditorState& B)
+	{
+		return A.GetCategoryText().CompareTo(B.GetCategoryText()) < 0;
+	});
+
+	for (UEditorState* EditorStateObject : SortedEditorStateObjects)
+	{
+		IDetailCategoryBuilder& EditorStateObjectCategory = DetailBuilder.EditCategory(EditorStateObject->GetFName(), EditorStateObject->GetCategoryText());
+		EditorStateObjectCategory.AddExternalObjects({ EditorStateObject }, EPropertyLocation::Default, FAddPropertyParams().HideRootObjectNode(true).CreateCategoryNodes(false));
+	}
+}
+
+void FWorldBookmarkDetailsCustomization::RefreshCustomDetail()
+{
+	if (CachedDetailBuilder)
+	{
+		CachedDetailBuilder->ForceRefreshDetails();
+	}
+}
+
+void FWorldBookmarkDetailsCustomization::PostUndo(bool bSuccess)
+{
+	// Refresh the UI
+	RefreshCustomDetail();
+}
+
+void FWorldBookmarkDetailsCustomization::PostRedo(bool bSuccess)
+{
+	// Refresh the UI
+	RefreshCustomDetail();
+}
diff --git a/Engine/Source/Editor/WorldBookmark/Private/WorldBookmark/WorldBookmarkDetailsCustomization.h b/Engine/Source/Editor/WorldBookmark/Private/WorldBookmark/WorldBookmarkDetailsCustomization.h
new file mode 100644
--- /dev/null
+++ b/Engine/Source/Editor/WorldBookmark/Private/WorldBookmark/WorldBookmarkDetailsCustomization.h
@@ -0,0 +1,34 @@
+// Copyright Epic Games, Inc. All Rights Reserved.
+
+#pragma once
+
+#include "IDetailCustomization.h"
+#include "EditorUndoClient.h"
+
+class UWorldBookmark;
+
+/** UI customization for UWorldBookmark */
+class FWorldBookmarkDetailsCustomization : public IDetailCustomization, public FEditorUndoClient
+{
+public:
+	static TSharedRef<IDetailCustomization> MakeInstance();
+
+	FWorldBookmarkDetailsCustomization();
+	~FWorldBookmarkDetailsCustomization();
+
+protected:
+	//~ Begin IDetailCustomization Interface.
+	virtual void CustomizeDetails(IDetailLayoutBuilder& DetailBuilder) override;
+	//~ End IDetailCustomization Interface.
+
+	/** FEditorUndoClient interface */
+	virtual void PostUndo(bool bSuccess) override;
+	virtual void PostRedo(bool bSuccess) override;
+	/** End FEditorUndoClient interface */
+
+	void RefreshCustomDetail();
+
+private:
+	/** The detail builder for this cusomtomisation. */
+	IDetailLayoutBuilder* CachedDetailBuilder;
+};
diff --git a/Engine/Source/Editor/WorldBookmark/Private/WorldBookmark/WorldBookmarkFactory.cpp b/Engine/Source/Editor/WorldBookmark/Private/WorldBookmark/WorldBookmarkFactory.cpp
new file mode 100644
--- /dev/null
+++ b/Engine/Source/Editor/WorldBookmark/Private/WorldBookmark/WorldBookmarkFactory.cpp
@@ -0,0 +1,24 @@
+// Copyright Epic Games, Inc. All Rights Reserved.
+
+#include "WorldBookmark/WorldBookmarkFactory.h"
+#include "WorldBookmark/WorldBookmark.h"
+#include "Misc/AssertionMacros.h"
+#include "Templates/SubclassOf.h"
+
+
+UWorldBookmarkFactory::UWorldBookmarkFactory(const FObjectInitializer& ObjectInitializer)
+	: Super(ObjectInitializer)
+{
+	SupportedClass = UWorldBookmark::StaticClass();
+
+	bCreateNew = true;
+	bEditAfterNew = true;
+	bEditorImport = false;
+}
+
+UObject* UWorldBookmarkFactory::FactoryCreateNew(UClass* Class, UObject* InParent, FName Name, EObjectFlags Flags, UObject* Context, FFeedbackContext* Warn)
+{
+	UWorldBookmark* Bookmark = NewObject<UWorldBookmark>(InParent, Name, Flags);
+	check(Bookmark);
+	return Bookmark;
+}
diff --git a/Engine/Source/Editor/WorldBookmark/Private/WorldBookmark/WorldBookmarkFactory.h b/Engine/Source/Editor/WorldBookmark/Private/WorldBookmark/WorldBookmarkFactory.h
new file mode 100644
--- /dev/null
+++ b/Engine/Source/Editor/WorldBookmark/Private/WorldBookmark/WorldBookmarkFactory.h
@@ -0,0 +1,23 @@
+// Copyright Epic Games, Inc. All Rights Reserved.
+
+#pragma once
+
+#include "Factories/Factory.h"
+#include "UObject/NameTypes.h"
+#include "UObject/UObjectGlobals.h"
+
+#include "WorldBookmarkFactory.generated.h"
+
+class FFeedbackContext;
+class UClass;
+class UObject;
+
+UCLASS()
+class UWorldBookmarkFactory : public UFactory
+{
+	GENERATED_UCLASS_BODY()
+
+	//~ Begin UFactory Interface
+	virtual UObject* FactoryCreateNew(UClass* Class, UObject* InParent, FName Name, EObjectFlags Flags, UObject* Context, FFeedbackContext* Warn) override;
+	//~ Begin UFactory Interface	
+};
diff --git a/Engine/Source/Editor/WorldBookmark/Private/WorldBookmark/WorldBookmarkModule.cpp b/Engine/Source/Editor/WorldBookmark/Private/WorldBookmark/WorldBookmarkModule.cpp
new file mode 100644
--- /dev/null
+++ b/Engine/Source/Editor/WorldBookmark/Private/WorldBookmark/WorldBookmarkModule.cpp
@@ -0,0 +1,75 @@
+// Copyright Epic Games, Inc. All Rights Reserved.
+
+#include "WorldBookmark/WorldBookmarkModule.h"
+#include "WorldBookmark/WorldBookmark.h"
+#include "WorldBookmark/WorldBookmarkDetailsCustomization.h"
+#include "WorldBookmark/Browser/SWorldBookmarkBrowser.h"
+
+#include "LevelEditor.h"
+#include "Modules/ModuleManager.h"
+#include "PropertyEditorModule.h"
+#include "Widgets/Docking/SDockTab.h"
+#include "WorkspaceMenuStructure.h"
+#include "WorkspaceMenuStructureModule.h"
+
+IMPLEMENT_MODULE(FWorldBookmarkModule, WorldBookmark);
+
+#define LOCTEXT_NAMESPACE "WorldBookmark"
+
+const FName WorldBookmarkBrowserTabId("WorldBookmarkBrowser");
+
+void FWorldBookmarkModule::StartupModule()
+{
+	FLevelEditorModule& LevelEditorModule = FModuleManager::Get().LoadModuleChecked<FLevelEditorModule>("LevelEditor");
+	LevelEditorModule.OnRegisterTabs().AddRaw(this, &FWorldBookmarkModule::RegisterWorldBookmarkBrowserTab);
+
+	FPropertyEditorModule& PropertyModule = FModuleManager::LoadModuleChecked<FPropertyEditorModule>("PropertyEditor");
+	PropertyModule.RegisterCustomClassLayout(UWorldBookmark::StaticClass()->GetFName(), FOnGetDetailCustomizationInstance::CreateStatic(&FWorldBookmarkDetailsCustomization::MakeInstance));
+}
+
+void FWorldBookmarkModule::ShutdownModule()
+{
+	FLevelEditorModule& LevelEditorModule = FModuleManager::Get().LoadModuleChecked<FLevelEditorModule>("LevelEditor");
+	LevelEditorModule.OnRegisterTabs().RemoveAll(this);
+
+	if (TSharedPtr<FTabManager> TabManager = LevelEditorModule.GetLevelEditorTabManager())
+	{
+		TabManager->UnregisterTabSpawner(WorldBookmarkBrowserTabId);
+	}
+
+	FPropertyEditorModule& PropertyModule = FModuleManager::LoadModuleChecked<FPropertyEditorModule>("PropertyEditor");
+	PropertyModule.UnregisterCustomClassLayout(UWorldBookmark::StaticClass()->GetFName());
+}
+
+void FWorldBookmarkModule::RegisterWorldBookmarkBrowserTab(TSharedPtr<FTabManager> InTabManager)
+{
+	const IWorkspaceMenuStructure& MenuStructure = WorkspaceMenu::GetMenuStructure();
+
+	const FSlateIcon WorldPartitionIcon(FAppStyle::GetAppStyleSetName(), "WorldPartition.WorldBookmark.TabIcon");
+
+	InTabManager->RegisterTabSpawner(WorldBookmarkBrowserTabId,
+		FOnSpawnTab::CreateRaw(this, &FWorldBookmarkModule::SpawnWorldBookmarkBrowserTab))
+		.SetDisplayName(LOCTEXT("WorldBookmarks", "World Bookmarks"))
+		.SetTooltipText(LOCTEXT("WorldBookmarksTooltipText", "Open the World Bookmarks browser."))
+		.SetGroup(MenuStructure.GetLevelEditorWorldPartitionCategory())
+		.SetIcon(WorldPartitionIcon);
+}
+
+TSharedRef<SDockTab> FWorldBookmarkModule::SpawnWorldBookmarkBrowserTab(const FSpawnTabArgs& Args)
+{
+	TSharedRef<SDockTab> NewTab =
+		SNew(SDockTab)
+		.Label(LOCTEXT("WorldBookmarkBrowserTab", "World Bookmarks"))
+		[
+			CreateWorldBookmarkBrowser()
+		];
+
+	return NewTab;
+}
+
+TSharedRef<SWidget> FWorldBookmarkModule::CreateWorldBookmarkBrowser()
+{
+	return SNew(SWorldBookmarkBrowser);
+}
+
+#undef LOCTEXT_NAMESPACE
diff --git a/Engine/Source/Editor/WorldBookmark/Private/WorldBookmark/WorldBookmarkModule.h b/Engine/Source/Editor/WorldBookmark/Private/WorldBookmark/WorldBookmarkModule.h
new file mode 100644
--- /dev/null
+++ b/Engine/Source/Editor/WorldBookmark/Private/WorldBookmark/WorldBookmarkModule.h
@@ -0,0 +1,33 @@
+// Copyright Epic Games, Inc. All Rights Reserved.
+
+#pragma once
+
+#include "Modules/ModuleInterface.h"
+#include "Templates/SharedPointer.h"
+
+class FSpawnTabArgs;
+class FTabManager;
+class SDockTab;
+class SWidget;
+
+/**
+ * Module dedicated to the World Bookmark feature
+ */
+class FWorldBookmarkModule : public IModuleInterface
+{
+public:
+	/**
+	 * Called right after the module DLL has been loaded and the module object has been created
+	 */
+	virtual void StartupModule();
+
+	/**
+	 * Called before the module is unloaded, right before the module object is destroyed.
+	 */
+	virtual void ShutdownModule();
+
+private:
+	void RegisterWorldBookmarkBrowserTab(TSharedPtr<FTabManager> InTabManager);
+	TSharedRef<SDockTab> SpawnWorldBookmarkBrowserTab(const FSpawnTabArgs& Args);
+	TSharedRef<SWidget> CreateWorldBookmarkBrowser();
+};
diff --git a/Engine/Source/Editor/WorldBookmark/Public/WorldBookmark/WorldBookmark.h b/Engine/Source/Editor/WorldBookmark/Public/WorldBookmark/WorldBookmark.h
new file mode 100644
--- /dev/null
+++ b/Engine/Source/Editor/WorldBookmark/Public/WorldBookmark/WorldBookmark.h
@@ -0,0 +1,73 @@
+// Copyright Epic Games, Inc. All Rights Reserved.
+
+#pragma once
+
+#include "UObject/Interface.h"
+#include "UObject/ObjectSaveContext.h"
+#include "EditorState/EditorStateCollection.h"
+#include "WorldBookmark.generated.h"
+
+/**
+ * World Bookmarks are assets that stores the state of the editor world.
+ */
+UCLASS(MinimalAPI, Config = UserWorldBookmarks, PerObjectConfig)
+class UWorldBookmark : public UObject
+{
+	GENERATED_UCLASS_BODY()
+
+	/** Load the bookmark data, restoring the editor state (loaded world, camera location, etc) to what is defined by the bookmark. */
+	void Load();
+
+	/** Update the bookmark so that it reflects the current state of the editor. */
+	void Update();
+
+	/** Return true if this bookmark was flagged as being a favorite bookmark of the user. */
+	bool GetIsUserFavorite() const;
+
+	/** Mark this bookmark as being a favorite bookmark of the user. */
+	void SetIsUserFavorite(bool bIsUserFavorite);
+
+	/** Retrieve the last time that bookmark was loaded, in UTC. */
+	FDateTime GetUserLastLoadedTimeStampUTC() const;
+
+	/** Store the last time that bookmark was loaded, in UTC. */
+	void SetUserLastLoadedTimeStampUTC(const FDateTime& InLastLoadedTimeStampUTC);
+
+private:
+	friend class FWorldBookmarkDetailsCustomization;
+
+	//~ Begin UObject Interface
+	virtual void PostLoad() override;
+	virtual void PreSave(FObjectPreSaveContext ObjectSaveContext) override;
+	virtual void GetAssetRegistryTags(FAssetRegistryTagsContext Context) const override;
+	virtual void PostInitProperties() override;
+	virtual void PostDuplicate(bool bDuplicateForPIE) override;
+	virtual void OverridePerObjectConfigSection(FString& SectionName) override;
+	//~ End UObject Interface
+		
+private:
+	// State of the editor.
+	UPROPERTY()
+	FEditorStateCollection EditorState;
+
+	// BookmarkGuid is our key to fetch the user settings for this bookmark in the Bookmarks.ini config file.
+	// It will remain unique even if redirectors are created for the world or the bookmark itself.
+	UPROPERTY(DuplicateTransient)
+	FGuid BookmarkGuid;
+
+	// BEGIN - User settings saved to the config
+	
+	// Last loaded time (UTC)
+	UPROPERTY(Config, Transient)
+	FDateTime LastLoadedTimeStampUTC;
+
+	// User favorite
+	UPROPERTY(Config, Transient)
+	bool bFavorite;
+
+	// Unused - The sole purpose of this property is to help users who would want to investigate/make changes to the ini themselves, as GUIDs are pretty opaque.
+	UPROPERTY(Config, Transient)
+	FString BookmarkAssetPath;
+
+	// END - User settings saved to the config
+};
diff --git a/Engine/Source/Editor/WorldBookmark/WorldBookmark.Build.cs b/Engine/Source/Editor/WorldBookmark/WorldBookmark.Build.cs
new file mode 100644
--- /dev/null
+++ b/Engine/Source/Editor/WorldBookmark/WorldBookmark.Build.cs
@@ -0,0 +1,32 @@
+// Copyright Epic Games, Inc. All Rights Reserved.
+
+using UnrealBuildTool;
+
+public class WorldBookmark : ModuleRules
+{
+	public WorldBookmark(ReadOnlyTargetRules Target) : base(Target)
+	{
+		PrivateDependencyModuleNames.AddRange(
+			new string[] {
+				"AssetDefinition",
+				"AssetRegistry",
+				"ContentBrowser",
+				"ContentBrowserData",
+				"Core",
+				"CoreUObject",
+				"EditorSubsystem",
+				"Engine",
+				"InputCore",
+				"PropertyEditor",
+				"SceneOutliner",
+				"SourceControl",
+				"Slate",
+				"SlateCore",				
+				"ToolMenus",
+				"ToolWidgets",
+				"UnrealEd",
+				"WorkspaceMenuStructure"
+			}
+		);
+	}
+}
diff --git a/Engine/Source/Editor/WorldPartitionEditor/Private/WorldPartition/WorldPartitionEditorState.cpp b/Engine/Source/Editor/WorldPartitionEditor/Private/WorldPartition/WorldPartitionEditorState.cpp
new file mode 100644
--- /dev/null
+++ b/Engine/Source/Editor/WorldPartitionEditor/Private/WorldPartition/WorldPartitionEditorState.cpp
@@ -0,0 +1,126 @@
+// Copyright Epic Games, Inc. All Rights Reserved.
+
+#include "WorldPartition/WorldPartitionEditorState.h"
+#include "WorldPartition/LoaderAdapter/LoaderAdapterShape.h"
+#include "WorldPartition/WorldPartitionEditorLoaderAdapter.h"
+#include "WorldPartition/WorldPartitionEditorSettings.h"
+#include "Engine/World.h"
+#include "LocationVolume.h"
+
+#include UE_INLINE_GENERATED_CPP_BY_NAME(WorldPartitionEditorState)
+
+#define LOCTEXT_NAMESPACE "WorldPartitionEditorState"
+
+UWorldPartitionEditorState::UWorldPartitionEditorState(const FObjectInitializer& ObjectInitializer)
+	: Super(ObjectInitializer)
+{
+}
+
+FText UWorldPartitionEditorState::GetCategoryText() const
+{
+	return FText(LOCTEXT("WorldPartitionEditorStateCategoryText", "World Partition"));
+}
+
+bool UWorldPartitionEditorState::CaptureState()
+{
+	UWorld* CurrentWorld = GetStateWorld();
+	if (CurrentWorld && CurrentWorld->IsPartitionedWorld())
+	{
+		if (GetDefault<UWorldPartitionEditorSettings>()->GetEnableLoadingInEditor())
+		{
+			auto GetRegionHash = [](const FBox& LoadedRegion)
+			{
+				return HashCombineFast(GetTypeHash(LoadedRegion.Min), GetTypeHash(LoadedRegion.Max));
+			};
+
+			// Loaded regions
+			TSet<uint32> LoadedEditorRegionsHashes;
+			for (const FBox& LoadedRegion : CurrentWorld->GetWorldPartition()->GetUserLoadedEditorRegions())
+			{
+				bool bAlreadyInSet = false;
+				LoadedEditorRegionsHashes.FindOrAdd(GetRegionHash(LoadedRegion), &bAlreadyInSet);
+				if (!bAlreadyInSet)
+				{
+					LoadedEditorRegions.Add(LoadedRegion);
+				}
+			}
+
+			// Loaded location volumes
+			for (FActorDescContainerInstanceCollection::TConstIterator<> Iterator(CurrentWorld->GetWorldPartition()); Iterator; ++Iterator)
+			{
+				if (ALocationVolume* LocationVolume = Cast<ALocationVolume>(Iterator->GetActor()); IsValid(LocationVolume))
+				{
+					check(LocationVolume->GetClass()->ImplementsInterface(UWorldPartitionActorLoaderInterface::StaticClass()));
+
+					IWorldPartitionActorLoaderInterface::ILoaderAdapter* LoaderAdapter = Cast<IWorldPartitionActorLoaderInterface>(LocationVolume)->GetLoaderAdapter();
+					check(LoaderAdapter);
+
+					if (LoaderAdapter->IsLoaded() && LoaderAdapter->GetUserCreated())
+					{
+						LoadedEditorLocationVolumes.Add(LocationVolume->GetFName());
+					}
+				}
+			}
+		}
+
+		return true;
+	}
+
+	return false;
+}
+
+bool UWorldPartitionEditorState::RestoreState() const
+{
+	UWorld* CurrentWorld = GetStateWorld();
+	if (CurrentWorld && CurrentWorld->IsPartitionedWorld())
+	{
+		if (GetDefault<UWorldPartitionEditorSettings>()->GetEnableLoadingInEditor())
+		{
+			auto GetRegionHash = [](const FBox& LoadedRegion)
+			{
+				// Round region's min/max to avoid precision issues.
+				FInt64Vector LoadedRegionMin(FMath::RoundToInt(LoadedRegion.Min.X), FMath::RoundToInt(LoadedRegion.Min.Y), FMath::RoundToInt(LoadedRegion.Min.Z));
+				FInt64Vector LoadedRegionMax(FMath::RoundToInt(LoadedRegion.Max.X), FMath::RoundToInt(LoadedRegion.Max.Y), FMath::RoundToInt(LoadedRegion.Max.Z));
+				return HashCombineFast(GetTypeHash(LoadedRegionMin), GetTypeHash(LoadedRegionMax));
+			};
+
+			// Grab existing loaded regions to avoid creating duplicate regions
+			TSet<uint32> LoadedEditorRegionsHashes;
+			for (const FBox& LoadedRegion : CurrentWorld->GetWorldPartition()->GetUserLoadedEditorRegions())
+			{
+				LoadedEditorRegionsHashes.FindOrAdd(GetRegionHash(LoadedRegion));
+			}
+
+			// Loaded regions
+			for (const FBox& LoadedRegion : LoadedEditorRegions)
+			{
+				if (LoadedRegion.IsValid)
+				{
+					bool bAlreadyInSet = false;
+					LoadedEditorRegionsHashes.FindOrAdd(GetRegionHash(LoadedRegion), &bAlreadyInSet);
+					if (!bAlreadyInSet)
+					{
+						UWorldPartitionEditorLoaderAdapter* EditorLoaderAdapter = CurrentWorld->GetWorldPartition()->CreateEditorLoaderAdapter<FLoaderAdapterShape>(CurrentWorld, LoadedRegion, TEXT("Loaded Region"));
+						IWorldPartitionActorLoaderInterface::ILoaderAdapter* LoaderAdapter = EditorLoaderAdapter->GetLoaderAdapter();
+						check(LoaderAdapter);
+						LoaderAdapter->SetUserCreated(true);
+						LoaderAdapter->Load();
+					}
+				}
+			}
+
+			// Location volumes
+			for (const FName& EditorLoadedLocationVolume : LoadedEditorLocationVolumes)
+			{
+				if (ALocationVolume* LocationVolume = FindObject<ALocationVolume>(CurrentWorld->PersistentLevel, *EditorLoadedLocationVolume.ToString()))
+				{
+					LocationVolume->Load();
+				}
+			}
+		}
+	}
+
+	return true;
+}
+
+#undef LOCTEXT_NAMESPACE
diff --git a/Engine/Source/Editor/WorldPartitionEditor/Private/WorldPartitionEditorModule.cpp b/Engine/Source/Editor/WorldPartitionEditor/Private/WorldPartitionEditorModule.cpp
--- a/Engine/Source/Editor/WorldPartitionEditor/Private/WorldPartitionEditorModule.cpp
+++ b/Engine/Source/Editor/WorldPartitionEditor/Private/WorldPartitionEditorModule.cpp
@@ -65,6 +65,9 @@
 #include "UObject/UObjectBase.h"
 #include "WorldPartition/ContentBundle/ContentBundleEditor.h"
 
+#include "EditorState/EditorStateSubsystem.h"
+#include "WorldPartition/WorldPartitionEditorState.h"
+
 IMPLEMENT_MODULE( FWorldPartitionEditorModule, WorldPartitionEditor );
 
 #define LOCTEXT_NAMESPACE "WorldPartition"
@@ -251,7 +254,15 @@ void FWorldPartitionEditorModule::StartupModule()
 				LevelEditorModule.AddCustomClassFilterToOutliner(ClassFilterData);
 			}
 		}
-	});
+
+		UEditorStateSubsystem::RegisterEditorStateType<UWorldPartitionEditorState>();
+
+		EditorCloseHandle = GEditor->OnEditorClose().AddLambda([this]()
+		{
+			UEditorStateSubsystem::UnregisterEditorStateType<UWorldPartitionEditorState>();
+			GEditor->OnEditorClose().Remove(EditorCloseHandle);
+		});
+	});		
 
 	IAssetReferenceFilter::OnIsCrossPluginReferenceAllowed().BindRaw(this, &FWorldPartitionEditorModule::OnIsCrossPluginReferenceAllowed);
 }
diff --git a/Engine/Source/Editor/WorldPartitionEditor/Public/WorldPartition/WorldPartitionEditorState.h b/Engine/Source/Editor/WorldPartitionEditor/Public/WorldPartition/WorldPartitionEditorState.h
new file mode 100644
--- /dev/null
+++ b/Engine/Source/Editor/WorldPartitionEditor/Public/WorldPartition/WorldPartitionEditorState.h
@@ -0,0 +1,26 @@
+// Copyright Epic Games, Inc. All Rights Reserved.
+
+#pragma once
+
+#include "EditorState/WorldEditorState.h"
+#include "WorldPartitionEditorState.generated.h"
+
+UCLASS(MinimalAPI)
+class UWorldPartitionEditorState : public UWorldDependantEditorState
+{
+	GENERATED_UCLASS_BODY()
+
+	//~ Begin UEditorState Interface
+	virtual FText GetCategoryText() const override;
+private:
+	virtual bool CaptureState() override;
+	virtual bool RestoreState() const override;
+	//~ End UEditorState Interface
+
+private:
+	UPROPERTY(EditAnywhere, Category = WorldState)
+	TArray<FBox> LoadedEditorRegions;
+
+	UPROPERTY(EditAnywhere, Category = WorldState)
+	TArray<FName> LoadedEditorLocationVolumes;
+};
diff --git a/Engine/Source/Editor/WorldPartitionEditor/Public/WorldPartitionEditorModule.h b/Engine/Source/Editor/WorldPartitionEditor/Public/WorldPartitionEditorModule.h
--- a/Engine/Source/Editor/WorldPartitionEditor/Public/WorldPartitionEditorModule.h
+++ b/Engine/Source/Editor/WorldPartitionEditor/Public/WorldPartitionEditorModule.h
@@ -167,6 +167,7 @@ private:
 	bool OnIsCrossPluginReferenceAllowed(const FAssetData& ReferencingAssetData, const FAssetData& ReferencedAssetData);
 
 	FDelegateHandle EditorInitializedHandle;
+	FDelegateHandle EditorCloseHandle;
 	FDelegateHandle LevelEditorExtenderDelegateHandle;
 
 	TWeakPtr<SDockTab> WorldPartitionTab;
