Fixed warnings in Build Engine Localization
* LOCTEXT macro has an empty source text and cannot be gathered.

#rnx

diff --git a/Engine/Source/Editor/WorldBookmark/Private/WorldBookmark/Browser/SWorldBookmarkBrowser.cpp b/Engine/Source/Editor/WorldBookmark/Private/WorldBookmark/Browser/SWorldBookmarkBrowser.cpp
--- a/Engine/Source/Editor/WorldBookmark/Private/WorldBookmark/Browser/SWorldBookmarkBrowser.cpp
+++ b/Engine/Source/Editor/WorldBookmark/Private/WorldBookmark/Browser/SWorldBookmarkBrowser.cpp
@@ -25,14 +25,14 @@ namespace WorldBookmarkBrowserColumn
 	namespace Favorite
 	{
 		FName Id = TEXT("Favorite");
-		FText DisplayText = LOCTEXT("Favorite_Text", "");
+		FText DisplayText = FText::GetEmpty();
 		FText ToolTipText = LOCTEXT("Favorite_Tooltip", "Favorite");
 	};
 
 	namespace RecentlyUsed
 	{
 		FName Id = TEXT("RecentlyUsed");
-		FText DisplayText = LOCTEXT("RecentlyUsed_Text", "");
+		FText DisplayText = FText::GetEmpty();
 		FText ToolTipText = LOCTEXT("RecentlyUsed_Tooltip", "Recently Used");
 	};
 
