World Bookmark - Prevent the WorldBookmark module from registering editor hooks when running unattended processes

diff --git a/Engine/Source/Editor/WorldBookmark/Private/WorldBookmark/WorldBookmarkModule.cpp b/Engine/Source/Editor/WorldBookmark/Private/WorldBookmark/WorldBookmarkModule.cpp
--- a/Engine/Source/Editor/WorldBookmark/Private/WorldBookmark/WorldBookmarkModule.cpp
+++ b/Engine/Source/Editor/WorldBookmark/Private/WorldBookmark/WorldBookmarkModule.cpp
@@ -20,25 +20,31 @@ const FName WorldBookmarkBrowserTabId("WorldBookmarkBrowser");
 
 void FWorldBookmarkModule::StartupModule()
 {
-	FLevelEditorModule& LevelEditorModule = FModuleManager::Get().LoadModuleChecked<FLevelEditorModule>("LevelEditor");
-	LevelEditorModule.OnRegisterTabs().AddRaw(this, &FWorldBookmarkModule::RegisterWorldBookmarkBrowserTab);
+	if (GIsEditor && !IsRunningCommandlet())
+	{
+		FLevelEditorModule& LevelEditorModule = FModuleManager::Get().LoadModuleChecked<FLevelEditorModule>("LevelEditor");
+		LevelEditorModule.OnRegisterTabs().AddRaw(this, &FWorldBookmarkModule::RegisterWorldBookmarkBrowserTab);
 
-	FPropertyEditorModule& PropertyModule = FModuleManager::LoadModuleChecked<FPropertyEditorModule>("PropertyEditor");
-	PropertyModule.RegisterCustomClassLayout(UWorldBookmark::StaticClass()->GetFName(), FOnGetDetailCustomizationInstance::CreateStatic(&FWorldBookmarkDetailsCustomization::MakeInstance));
+		FPropertyEditorModule& PropertyModule = FModuleManager::LoadModuleChecked<FPropertyEditorModule>("PropertyEditor");
+		PropertyModule.RegisterCustomClassLayout(UWorldBookmark::StaticClass()->GetFName(), FOnGetDetailCustomizationInstance::CreateStatic(&FWorldBookmarkDetailsCustomization::MakeInstance));
+	}
 }
 
 void FWorldBookmarkModule::ShutdownModule()
 {
-	FLevelEditorModule& LevelEditorModule = FModuleManager::Get().LoadModuleChecked<FLevelEditorModule>("LevelEditor");
-	LevelEditorModule.OnRegisterTabs().RemoveAll(this);
-
-	if (TSharedPtr<FTabManager> TabManager = LevelEditorModule.GetLevelEditorTabManager())
+	if (GIsEditor && !IsRunningCommandlet())
 	{
-		TabManager->UnregisterTabSpawner(WorldBookmarkBrowserTabId);
-	}
+		FLevelEditorModule& LevelEditorModule = FModuleManager::Get().LoadModuleChecked<FLevelEditorModule>("LevelEditor");
+		LevelEditorModule.OnRegisterTabs().RemoveAll(this);
+
+		if (TSharedPtr<FTabManager> TabManager = LevelEditorModule.GetLevelEditorTabManager())
+		{
+			TabManager->UnregisterTabSpawner(WorldBookmarkBrowserTabId);
+		}
 
-	FPropertyEditorModule& PropertyModule = FModuleManager::LoadModuleChecked<FPropertyEditorModule>("PropertyEditor");
-	PropertyModule.UnregisterCustomClassLayout(UWorldBookmark::StaticClass()->GetFName());
+		FPropertyEditorModule& PropertyModule = FModuleManager::LoadModuleChecked<FPropertyEditorModule>("PropertyEditor");
+		PropertyModule.UnregisterCustomClassLayout(UWorldBookmark::StaticClass()->GetFName());
+	}
 }
 
 void FWorldBookmarkModule::RegisterWorldBookmarkBrowserTab(TSharedPtr<FTabManager> InTabManager)
