World Bookmark - Fixed issue reported by ASAN where UWorldBookmark::StaticClass() is accessed during the WorldBookmark module shutdown

diff --git a/Engine/Source/Editor/WorldBookmark/Private/WorldBookmark/WorldBookmarkModule.cpp b/Engine/Source/Editor/WorldBookmark/Private/WorldBookmark/WorldBookmarkModule.cpp
--- a/Engine/Source/Editor/WorldBookmark/Private/WorldBookmark/WorldBookmarkModule.cpp
+++ b/Engine/Source/Editor/WorldBookmark/Private/WorldBookmark/WorldBookmarkModule.cpp
@@ -26,6 +26,8 @@ void FWorldBookmarkModule::StartupModule()
 		LevelEditorModule.OnRegisterTabs().AddRaw(this, &FWorldBookmarkModule::RegisterWorldBookmarkBrowserTab);
 
 		FPropertyEditorModule& PropertyModule = FModuleManager::LoadModuleChecked<FPropertyEditorModule>("PropertyEditor");
+
+		ClassesToUnregisterOnShutdown.Add(UWorldBookmark::StaticClass()->GetFName());
 		PropertyModule.RegisterCustomClassLayout(UWorldBookmark::StaticClass()->GetFName(), FOnGetDetailCustomizationInstance::CreateStatic(&FWorldBookmarkDetailsCustomization::MakeInstance));
 	}
 }
@@ -34,16 +36,25 @@ void FWorldBookmarkModule::ShutdownModule()
 {
 	if (GIsEditor && !IsRunningCommandlet())
 	{
-		FLevelEditorModule& LevelEditorModule = FModuleManager::Get().LoadModuleChecked<FLevelEditorModule>("LevelEditor");
-		LevelEditorModule.OnRegisterTabs().RemoveAll(this);
+		if (FLevelEditorModule* LevelEditorModule = FModuleManager::Get().GetModulePtr<FLevelEditorModule>("LevelEditor"))
+		{
+			LevelEditorModule->OnRegisterTabs().RemoveAll(this);
 
-		if (TSharedPtr<FTabManager> TabManager = LevelEditorModule.GetLevelEditorTabManager())
+			if (TSharedPtr<FTabManager> TabManager = LevelEditorModule->GetLevelEditorTabManager())
+			{
+				TabManager->UnregisterTabSpawner(WorldBookmarkBrowserTabId);
+			}
+		}		
+
+		if (FPropertyEditorModule* PropertyModule = FModuleManager::GetModulePtr<FPropertyEditorModule>("PropertyEditor"))
 		{
-			TabManager->UnregisterTabSpawner(WorldBookmarkBrowserTabId);
+			for (const FName& ClassName : ClassesToUnregisterOnShutdown)
+			{
+				PropertyModule->UnregisterCustomClassLayout(ClassName);
+			}
 		}
 
-		FPropertyEditorModule& PropertyModule = FModuleManager::LoadModuleChecked<FPropertyEditorModule>("PropertyEditor");
-		PropertyModule.UnregisterCustomClassLayout(UWorldBookmark::StaticClass()->GetFName());
+		ClassesToUnregisterOnShutdown.Empty();
 	}
 }
 
diff --git a/Engine/Source/Editor/WorldBookmark/Private/WorldBookmark/WorldBookmarkModule.h b/Engine/Source/Editor/WorldBookmark/Private/WorldBookmark/WorldBookmarkModule.h
--- a/Engine/Source/Editor/WorldBookmark/Private/WorldBookmark/WorldBookmarkModule.h
+++ b/Engine/Source/Editor/WorldBookmark/Private/WorldBookmark/WorldBookmarkModule.h
@@ -30,4 +30,6 @@ private:
 	void RegisterWorldBookmarkBrowserTab(TSharedPtr<FTabManager> InTabManager);
 	TSharedRef<SDockTab> SpawnWorldBookmarkBrowserTab(const FSpawnTabArgs& Args);
 	TSharedRef<SWidget> CreateWorldBookmarkBrowser();
+
+	TArray<FName> ClassesToUnregisterOnShutdown;
 };
