World Bookmark - Use the editor world context world rather than GWorld

diff --git a/Engine/Source/Editor/UnrealEd/Private/EditorState/WorldEditorState.cpp b/Engine/Source/Editor/UnrealEd/Private/EditorState/WorldEditorState.cpp
--- a/Engine/Source/Editor/UnrealEd/Private/EditorState/WorldEditorState.cpp
+++ b/Engine/Source/Editor/UnrealEd/Private/EditorState/WorldEditorState.cpp
@@ -3,6 +3,7 @@
 #include "EditorState/WorldEditorState.h"
 #include "AssetRegistry/IAssetRegistry.h"
 #include "Engine/World.h"
+#include "Editor.h"
 #include "FileHelpers.h"
 #include "Misc/MessageDialog.h"
 
@@ -27,13 +28,14 @@ TSoftObjectPtr<UWorld> UWorldEditorState::GetStateWorld() const
 
 UEditorState::FOperationResult UWorldEditorState::CaptureState()
 {
-	if (!GWorld || FPackageName::IsTempPackage(GWorld->GetPackage()->GetName()))
+	UWorld* CurrentWorld = GEditor->GetEditorWorldContext().World();
+	if (!CurrentWorld || FPackageName::IsTempPackage(CurrentWorld->GetPackage()->GetName()))
 	{
 		return FOperationResult(FOperationResult::Failure, LOCTEXT("CaptureStateFailure_UnsavedWorld", "Current world must be saved"));
 	}
 
 	// Set the current world.
-	World = (UWorld*)GWorld;
+	World = CurrentWorld;
 	return FOperationResult(FOperationResult::Success);
 }
 
@@ -59,7 +61,7 @@ UEditorState::FOperationResult UWorldEditorState::RestoreState() const
 	}
 
 	// Already in the right world.
-	UWorld* CurrentWorld = GWorld;
+	UWorld* CurrentWorld = GEditor->GetEditorWorldContext().World();
 	if (CurrentWorld == World)
 	{
 		return FOperationResult(FOperationResult::Success, FText::Format(LOCTEXT("RestoreStateSuccess_WorldAlreadyLoaded", "Bookmark world {0} was already active"), WorldPackageText));
@@ -111,7 +113,7 @@ UWorldDependantEditorState::UWorldDependantEditorState(const FObjectInitializer&
 
 UWorld* UWorldDependantEditorState::GetStateWorld() const
 {
-	return GWorld;
+	return GEditor->GetEditorWorldContext().World();
 }
 
 TArray<TSubclassOf<UEditorState>> UWorldDependantEditorState::GetDependencies() const
diff --git a/Engine/Source/Editor/WorldBookmark/Private/WorldBookmark/WorldBookmark.cpp b/Engine/Source/Editor/WorldBookmark/Private/WorldBookmark/WorldBookmark.cpp
--- a/Engine/Source/Editor/WorldBookmark/Private/WorldBookmark/WorldBookmark.cpp
+++ b/Engine/Source/Editor/WorldBookmark/Private/WorldBookmark/WorldBookmark.cpp
@@ -16,6 +16,8 @@
 
 #include UE_INLINE_GENERATED_CPP_BY_NAME(WorldBookmark)
 
+DEFINE_LOG_CATEGORY_STATIC(LogWorldBookmark, All, All);
+
 #define LOCTEXT_NAMESPACE "WorldBookmark"
 
 UWorldBookmark::UWorldBookmark(const FObjectInitializer& ObjectInitializer)
@@ -86,6 +88,16 @@ void UWorldBookmark::GetAssetRegistryTags(FAssetRegistryTagsContext Context) con
 
 bool UWorldBookmark::CanLoad(FText* FailureReason) const
 {
+	if (!GIsEditor)
+	{
+		if (FailureReason)
+		{
+			*FailureReason = LOCTEXT("CanLoadFailure_NotInEditor", "Can't load bookmark outside of the editor");
+		}
+
+		return false;
+	}
+
 	if (GEditor->PlayWorld || GIsPlayInEditorWorld)
 	{
 		if (FailureReason)
@@ -111,19 +123,34 @@ void UWorldBookmark::Load()
 	}
 	else
 	{
-		FNotificationInfo Info(LOCTEXT("CantLoadBookmark", "Failed to load bookmark!"));
-		Info.SubText = CanLoadFailureReason;
-		Info.ExpireDuration = 5.0f;
-		Info.bFireAndForget = true;
-		Info.bUseLargeFont = false;
-		Info.Image = FCoreStyle::Get().GetBrush(TEXT("MessageLog.Warning"));
-		FSlateNotificationManager::Get().AddNotification(Info);
+		if (GIsEditor)
+		{
+			FNotificationInfo Info(LOCTEXT("CantLoadBookmark", "Failed to load bookmark!"));
+			Info.SubText = CanLoadFailureReason;
+			Info.ExpireDuration = 5.0f;
+			Info.bFireAndForget = true;
+			Info.bUseLargeFont = false;
+			Info.Image = FCoreStyle::Get().GetBrush(TEXT("MessageLog.Warning"));
+			FSlateNotificationManager::Get().AddNotification(Info);
+		}
+		
+		UE_LOG(LogWorldBookmark, Error, TEXT("Failed to load bookmark: %s"), *CanLoadFailureReason.ToString());
 	}
 }
 
 bool UWorldBookmark::CanUpdate(FText* FailureReason) const
 {
-	UWorld* CurrentWorld = GWorld;
+	if (!GIsEditor)
+	{
+		if (FailureReason)
+		{
+			*FailureReason = LOCTEXT("CanUpdateFailure_NotInEditor", "Can't update bookmark outside of the editor");
+		}
+
+		return false;
+	}
+
+	UWorld* CurrentWorld = GEditor->GetEditorWorldContext().World();
 	if (CurrentWorld == nullptr)
 	{
 		if (FailureReason)
@@ -172,13 +199,18 @@ void UWorldBookmark::Update()
 	}
 	else
 	{
-		FNotificationInfo Info(LOCTEXT("CantUpdateBookmark", "Failed to update bookmark!"));
-		Info.SubText = CanUpdateFailureReason;
-		Info.ExpireDuration = 5.0f;
-		Info.bFireAndForget = true;
-		Info.bUseLargeFont = false;
-		Info.Image = FCoreStyle::Get().GetBrush(TEXT("MessageLog.Warning"));
-		FSlateNotificationManager::Get().AddNotification(Info);
+		if (GIsEditor)
+		{
+			FNotificationInfo Info(LOCTEXT("CantUpdateBookmark", "Failed to update bookmark!"));
+			Info.SubText = CanUpdateFailureReason;
+			Info.ExpireDuration = 5.0f;
+			Info.bFireAndForget = true;
+			Info.bUseLargeFont = false;
+			Info.Image = FCoreStyle::Get().GetBrush(TEXT("MessageLog.Warning"));
+			FSlateNotificationManager::Get().AddNotification(Info);
+		}
+
+		UE_LOG(LogWorldBookmark, Error, TEXT("Failed to update bookmark: %s"), *CanUpdateFailureReason.ToString());
 	}
 }
 
