World Bookmarks - Added support for renaming & deleting bookmarks from the Bookmark browser
* Can be renamed inline (double selection, F2) or through the context menu
* Can be deleted with the delete key or through the context menu

diff --git a/Engine/Source/Editor/WorldBookmark/Private/WorldBookmark/Browser/SWorldBookmarkBrowser.cpp b/Engine/Source/Editor/WorldBookmark/Private/WorldBookmark/Browser/SWorldBookmarkBrowser.cpp
--- a/Engine/Source/Editor/WorldBookmark/Private/WorldBookmark/Browser/SWorldBookmarkBrowser.cpp
+++ b/Engine/Source/Editor/WorldBookmark/Private/WorldBookmark/Browser/SWorldBookmarkBrowser.cpp
@@ -5,16 +5,21 @@
 #include "AssetRegistry/AssetData.h"
 #include "AssetRegistry/AssetRegistryHelpers.h"
 #include "AssetToolsModule.h"
+#include "AssetViewUtils.h"
 #include "ContentBrowserDataSource.h"
 #include "ContentBrowserModule.h"
 #include "Editor.h"
+#include "Framework/Application/SlateApplication.h"
+#include "Framework/Commands/GenericCommands.h"
 #include "LevelEditor.h"
 #include "Modules/ModuleManager.h"
+#include "ObjectTools.h"
 #include "PropertyEditorModule.h"
 #include "SSimpleButton.h"
 #include "Toolkits/GlobalEditorCommonCommands.h"
 #include "UObject/AssetRegistryTagsContext.h"
 #include "Widgets/Input/SSearchBox.h"
+#include "Widgets/Text/SInlineEditableTextBlock.h"
 #include "WorldBookmark/WorldBookmark.h"
 #include "WorldBookmark/WorldBookmarkFactory.h"
 
@@ -57,6 +62,8 @@ static UWorld* GetCurrentWorld()
 }
 
 SWorldBookmarkBrowser::SWorldBookmarkBrowser()
+	: bIsInTick(false)
+	, bPendingRefresh(false)
 {
 	// Listen for asset registry updates
 	IAssetRegistry& AssetRegistry = IAssetRegistry::GetChecked();
@@ -152,6 +159,7 @@ void SWorldBookmarkBrowser::Construct(const FArguments& InArgs)
 		.OnSelectionChanged(this, &SWorldBookmarkBrowser::OnSelectedBookmarkChanged)
 		.OnMouseButtonDoubleClick(this, &SWorldBookmarkBrowser::OnBookmarkDoubleClicked)
 		.OnContextMenuOpening(this, &SWorldBookmarkBrowser::OnContextMenuOpening)
+		.OnItemScrolledIntoView(this, &SWorldBookmarkBrowser::ItemScrolledIntoView)
 		.HeaderRow
 		(
 			HeaderRow
@@ -223,14 +231,43 @@ void SWorldBookmarkBrowser::Construct(const FArguments& InArgs)
 	];
 
 	Commands = MakeShareable(new FUICommandList());
-	Commands->MapAction(FGlobalEditorCommonCommands::Get().FindInContentBrowser, FUIAction(
+	Commands->MapAction(FGlobalEditorCommonCommands::Get().FindInContentBrowser,
 		FExecuteAction::CreateSP(this, &SWorldBookmarkBrowser::FindSelectedBookmarkInContentBrowser),
-		FCanExecuteAction::CreateSP(this, &SWorldBookmarkBrowser::IsValidBookmarkSelected)
-	));
+		FCanExecuteAction::CreateSP(this, &SWorldBookmarkBrowser::IsValidBookmarkSelected));
+
+	Commands->MapAction(FGenericCommands::Get().Rename,
+		FExecuteAction::CreateSP(this, &SWorldBookmarkBrowser::RequestRenameSelectedBookmark),
+		FCanExecuteAction::CreateSP(this, &SWorldBookmarkBrowser::IsValidBookmarkSelected));
+
+	Commands->MapAction(FGenericCommands::Get().Delete,
+		FExecuteAction::CreateSP(this, &SWorldBookmarkBrowser::DeleteSelectedBookmark),
+		FCanExecuteAction::CreateSP(this, &SWorldBookmarkBrowser::IsValidBookmarkSelected));
+}
+
+FReply SWorldBookmarkBrowser::OnKeyDown(const FGeometry& MyGeometry, const FKeyEvent& InKeyEvent)
+{
+	bool bIsRenamingAsset = false;
+	if (FWorldBookmarkTreeItemPtr SelectedBookmarkTreeItem = GetSelectedBookmarkTreeItem())
+	{
+		bIsRenamingAsset = SelectedBookmarkTreeItem->bInEditingMode;
+	}
+		
+	if (bIsRenamingAsset || Commands->ProcessCommandBindings(InKeyEvent))
+	{
+		return FReply::Handled();
+	}
+	return FReply::Unhandled();
 }
 
 void SWorldBookmarkBrowser::RefreshItems()
 {
+	// Never perform a refresh outside of the tick method
+	if (!bIsInTick)
+	{
+		bPendingRefresh = true;
+		return;
+	}
+
 	IAssetRegistry& AssetRegistry = IAssetRegistry::GetChecked();
 	TArray<FAssetData> AssetsData;
 	const bool bSearchSubClasses = true;
@@ -272,6 +309,8 @@ void SWorldBookmarkBrowser::RefreshItems()
 
 	// Restore selected bookmark after refresh
 	SetSelectedBookmark(PreviousSelection);
+
+	bPendingRefresh = false;
 }
 
 void SWorldBookmarkBrowser::SortItems()
@@ -338,11 +377,11 @@ void FWorldBookmarkTableRow::Construct(const FArguments& InArgs, const TSharedRe
 
 TSharedRef<SWidget> FWorldBookmarkTableRow::GenerateWidgetForColumn(const FName& ColumnName)
 {
-	if (WeakWorldBookmarkTreeItem.IsValid())
+	if (FWorldBookmarkTreeItemPtr WorldBookmarkTreeItem = WeakWorldBookmarkTreeItem.Pin())
 	{
 		if (ColumnName == WorldBookmarkBrowserColumn::Favorite::Id)
 		{
-			if (UWorldBookmark* WorldBookmark = Cast<UWorldBookmark>(WeakWorldBookmarkTreeItem.Pin()->BookmarkAsset.GetAsset()))
+			if (UWorldBookmark* WorldBookmark = Cast<UWorldBookmark>(WorldBookmarkTreeItem->BookmarkAsset.GetAsset()))
 			{
 				if (WorldBookmark->GetIsUserFavorite())
 				{
@@ -361,7 +400,7 @@ TSharedRef<SWidget> FWorldBookmarkTableRow::GenerateWidgetForColumn(const FName&
 		}
 		else if (ColumnName == WorldBookmarkBrowserColumn::RecentlyUsed::Id)
 		{
-			if (UWorldBookmark* WorldBookmark = Cast<UWorldBookmark>(WeakWorldBookmarkTreeItem.Pin()->BookmarkAsset.GetAsset()))
+			if (UWorldBookmark* WorldBookmark = Cast<UWorldBookmark>(WorldBookmarkTreeItem->BookmarkAsset.GetAsset()))
 			{
 				FTimespan LastLoaded = FDateTime::UtcNow() - WorldBookmark->GetUserLastLoadedTimeStampUTC();
 				if (LastLoaded < FTimespan::FromDays(14))
@@ -381,35 +420,105 @@ TSharedRef<SWidget> FWorldBookmarkTableRow::GenerateWidgetForColumn(const FName&
 		}
 		else if (ColumnName == WorldBookmarkBrowserColumn::Label::Id)
 		{
-			return SNew(SVerticalBox)
-			
-			+ SVerticalBox::Slot()
-			.Padding(10, 5, 0, 0)
-			.AutoHeight()
-			[
-				SNew(STextBlock)
-					.Text(FText::FromName(WeakWorldBookmarkTreeItem.Pin()->BookmarkAsset.AssetName))
-			]
+			TSharedPtr<SInlineEditableTextBlock> InlineEditableTextBlock;
 
-			+ SVerticalBox::Slot()
-			.Padding(10, 0, 0, 5)
-			.AutoHeight()
-			[
-				SNew(STextBlock)
-					.ColorAndOpacity(FSlateColor::UseSubduedForeground())
-					.Text(this, &FWorldBookmarkTableRow::GetLastAccessedText)
-			];
+			TSharedRef<SWidget> LabelWidget =
+				SNew(SVerticalBox)
+				+ SVerticalBox::Slot()
+				.Padding(10, 5, 0, 0)
+				.AutoHeight()
+				[
+					SAssignNew(InlineEditableTextBlock, SInlineEditableTextBlock)
+						.Text(FText::FromName(WorldBookmarkTreeItem->BookmarkAsset.AssetName))
+						.OnVerifyTextChanged(this, &FWorldBookmarkTableRow::OnLabelTextVerifyChanged)
+						.OnTextCommitted(this, &FWorldBookmarkTableRow::OnLabelTextCommitted)
+						.OnEnterEditingMode(this, &FWorldBookmarkTableRow::OnEnterEditingMode)
+						.OnExitEditingMode(this, &FWorldBookmarkTableRow::OnExitEditingMode)
+						.IsSelected(this, &FWorldBookmarkTableRow::IsSelected)
+				]
+				+ SVerticalBox::Slot()
+				.Padding(10, 0, 0, 5)
+				.AutoHeight()
+				[
+					SNew(STextBlock)
+						.ColorAndOpacity(FSlateColor::UseSubduedForeground())
+						.Text(this, &FWorldBookmarkTableRow::GetLastAccessedText)
+				];
+
+			WorldBookmarkTreeItem->OnRenameRequested.BindSP(InlineEditableTextBlock.Get(), &SInlineEditableTextBlock::EnterEditingMode);
+
+			return LabelWidget;
 		}
 	}
 
 	return SNullWidget::NullWidget;
 }
 
+static FSoftObjectPath GetBookmarkObjectPath(FWorldBookmarkTreeItemPtr InWorldBookmarkTreeItem)
+{
+	return InWorldBookmarkTreeItem->BookmarkAsset.GetSoftObjectPath();
+}
+
+static FSoftObjectPath GetRenamedBookmarkObjectPath(FWorldBookmarkTreeItemPtr InWorldBookmarkTreeItem, const FText& InNewAssetName)
+{
+	const FString PackagePath = InWorldBookmarkTreeItem->BookmarkAsset.PackagePath.ToString();
+	const FString NewObjectPath = FString::Format(TEXT("{0}/{1}.{1}"), { *PackagePath, *InNewAssetName.ToString() });
+	return FSoftObjectPath(NewObjectPath);
+}
+
+bool FWorldBookmarkTableRow::OnLabelTextVerifyChanged(const FText& InNewText, FText& OutErrorMessage)
+{
+	if (FWorldBookmarkTreeItemPtr WorldBookmarkTreeItem = WeakWorldBookmarkTreeItem.Pin())
+	{
+		const FSoftObjectPath NewObjectPath = GetRenamedBookmarkObjectPath(WorldBookmarkTreeItem, InNewText);
+		return AssetViewUtils::IsValidObjectPathForCreate(NewObjectPath.ToString(), OutErrorMessage);
+	}
+	else
+	{
+		return false;
+	}
+}
+
+void FWorldBookmarkTableRow::OnLabelTextCommitted(const FText& InNewText, ETextCommit::Type InTextCommit)
+{
+	FText ErrorMessage;
+	if (OnLabelTextVerifyChanged(InNewText, ErrorMessage))
+	{
+		if (FWorldBookmarkTreeItemPtr WorldBookmarkTreeItem = WeakWorldBookmarkTreeItem.Pin())
+		{
+			FAssetToolsModule& AssetToolsModule = FModuleManager::LoadModuleChecked<FAssetToolsModule>("AssetTools");
+
+			TArray<FAssetRenameData> AssetsAndNames;
+			const FSoftObjectPath OldObjectPath = GetBookmarkObjectPath(WorldBookmarkTreeItem);
+			const FSoftObjectPath NewObjectPath = GetRenamedBookmarkObjectPath(WorldBookmarkTreeItem, InNewText);
+			new (AssetsAndNames) FAssetRenameData(OldObjectPath, NewObjectPath);
+
+			AssetToolsModule.Get().RenameAssetsWithDialog(AssetsAndNames);
+		}
+	}
+}
+
+void FWorldBookmarkTableRow::OnEnterEditingMode()
+{
+	if (FWorldBookmarkTreeItemPtr WorldBookmarkTreeItem = WeakWorldBookmarkTreeItem.Pin())
+	{
+		WorldBookmarkTreeItem->bInEditingMode = true;
+	}
+}
+
+void FWorldBookmarkTableRow::OnExitEditingMode()
+{
+	if (FWorldBookmarkTreeItemPtr WorldBookmarkTreeItem = WeakWorldBookmarkTreeItem.Pin())
+	{
+		WorldBookmarkTreeItem->bInEditingMode = false;
+	}
+}
+
 FText FWorldBookmarkTableRow::GetLastAccessedText() const
 {
-	if (WeakWorldBookmarkTreeItem.IsValid() && WeakWorldBookmarkTreeItem.Pin()->BookmarkAsset.IsValid())
+	if (FWorldBookmarkTreeItemPtr WorldBookmarkTreeItem = WeakWorldBookmarkTreeItem.Pin())
 	{
-		if (UWorldBookmark* WorldBookmark = Cast<UWorldBookmark>(WeakWorldBookmarkTreeItem.Pin()->BookmarkAsset.GetAsset()))
+		if (UWorldBookmark* WorldBookmark = Cast<UWorldBookmark>(WorldBookmarkTreeItem->BookmarkAsset.GetAsset()))
 		{
 			FDateTime LastLoadedUTC = WorldBookmark->GetUserLastLoadedTimeStampUTC();
 
@@ -595,18 +704,24 @@ TSharedPtr<SWidget> SWorldBookmarkBrowser::OnContextMenuOpening()
 	{
 		// Find In Content Browser
 		MenuBuilder.AddMenuEntry(FGlobalEditorCommonCommands::Get().FindInContentBrowser);
+
+		// Rename
+		MenuBuilder.AddMenuEntry(FGenericCommands::Get().Rename);
+
+		// Delete
+		MenuBuilder.AddMenuEntry(FGenericCommands::Get().Delete);
 	}
 	MenuBuilder.EndSection();
 
 	return MenuBuilder.MakeWidget();
 }
 
-bool IsWorldBookmarkAsset(const FAssetData& InAssetData)
+static bool IsWorldBookmarkAsset(const FAssetData& InAssetData)
 {
 	return InAssetData.AssetClassPath == UWorldBookmark::StaticClass()->GetClassPathName();
 }
 
-bool IsWorldBookmarkAsset(UObject* InObject)
+static bool IsWorldBookmarkAsset(UObject* InObject)
 {
 	return InObject->IsA<UWorldBookmark>();
 }
@@ -634,6 +749,16 @@ void SWorldBookmarkBrowser::OnAssetRenamed(const FAssetData& InAssetData, const
 {
 	if (IsWorldBookmarkAsset(InAssetData))
 	{
+		if (FWorldBookmarkTreeItemPtr SelectedBookmarkTreeItem = GetSelectedBookmarkTreeItem())
+		{
+			// If the selected tree item is the one that was renamed, assign the new AssetData to it
+			// This will allow it to remain selected after the view refreshes.
+			if (SelectedBookmarkTreeItem->BookmarkAsset.GetObjectPathString() == InOldObjectPath)
+			{
+				SelectedBookmarkTreeItem->BookmarkAsset = InAssetData;
+			}
+		}
+
 		RefreshItems();
 	}
 }
@@ -742,7 +867,7 @@ bool SWorldBookmarkBrowser::IsValidBookmarkSelected() const
 	return GetSelectedBookmark() != nullptr;
 }
 
-UWorldBookmark* SWorldBookmarkBrowser::GetSelectedBookmark() const
+FWorldBookmarkTreeItemPtr SWorldBookmarkBrowser::GetSelectedBookmarkTreeItem() const
 {
 	if (!BookmarksView)
 	{
@@ -756,7 +881,18 @@ UWorldBookmark* SWorldBookmarkBrowser::GetSelectedBookmark() const
 	}
 
 	FWorldBookmarkTreeItemPtr SelectedItem = Selection[0];
-	if (!SelectedItem.IsValid() || !SelectedItem->BookmarkAsset.IsValid())
+	if (!SelectedItem || !SelectedItem->BookmarkAsset.IsValid())
+	{
+		return nullptr;
+	}
+
+	return SelectedItem;
+}
+
+UWorldBookmark* SWorldBookmarkBrowser::GetSelectedBookmark() const
+{
+	FWorldBookmarkTreeItemPtr SelectedItem = GetSelectedBookmarkTreeItem();
+	if (!SelectedItem)
 	{
 		return nullptr;
 	}
@@ -790,5 +926,51 @@ void SWorldBookmarkBrowser::SetSelectedBookmark(UWorldBookmark* InBookmark)
 	}
 }
 
+void SWorldBookmarkBrowser::ItemScrolledIntoView(FWorldBookmarkTreeItemPtr WorldBookmarkTreeItem, const TSharedPtr<ITableRow>& Widget)
+{
+	if (BookmarkPendingRename && WorldBookmarkTreeItem == BookmarkPendingRename)
+	{
+		// Abort the rename if our widget has lost focus.
+		TSharedPtr<SWindow> OwnerWindow = FSlateApplication::Get().FindWidgetWindow(AsShared());
+		if (OwnerWindow && OwnerWindow->HasAnyUserFocusOrFocusedDescendants())
+		{
+			// We know the tree item is now visible, bring up the editable text block.
+			BookmarkPendingRename->OnRenameRequested.ExecuteIfBound();
+		}
+
+		BookmarkPendingRename.Reset();
+	}
+}
+
+void SWorldBookmarkBrowser::RequestRenameSelectedBookmark()
+{
+	if (FWorldBookmarkTreeItemPtr SelectedBookmarkTreeItem = GetSelectedBookmarkTreeItem())
+	{
+		// Before showing the inline editable text block, we must ensure that the tree item is visible.
+		BookmarkPendingRename = SelectedBookmarkTreeItem;
+		BookmarksView->RequestScrollIntoView(SelectedBookmarkTreeItem);
+	}
+}
+
+void SWorldBookmarkBrowser::DeleteSelectedBookmark()
+{
+	if (FWorldBookmarkTreeItemPtr SelectedBookmarkTreeItem = GetSelectedBookmarkTreeItem())
+	{
+		ObjectTools::DeleteAssets({ SelectedBookmarkTreeItem->BookmarkAsset });
+	}
+}
+
+void SWorldBookmarkBrowser::Tick(const FGeometry& AllottedGeometry, const double InCurrentTime, const float InDeltaTime)
+{
+	TGuardValue<bool> IsInTickGuard(bIsInTick, true);
+
+	if (bPendingRefresh)
+	{
+		RefreshItems();
+	}
+
+	SCompoundWidget::Tick(AllottedGeometry, InCurrentTime, InDeltaTime);
+}
+
 #undef LOCTEXT_NAMESPACE
 
diff --git a/Engine/Source/Editor/WorldBookmark/Private/WorldBookmark/Browser/SWorldBookmarkBrowser.h b/Engine/Source/Editor/WorldBookmark/Private/WorldBookmark/Browser/SWorldBookmarkBrowser.h
--- a/Engine/Source/Editor/WorldBookmark/Private/WorldBookmark/Browser/SWorldBookmarkBrowser.h
+++ b/Engine/Source/Editor/WorldBookmark/Private/WorldBookmark/Browser/SWorldBookmarkBrowser.h
@@ -18,6 +18,12 @@ struct FWorldBookmarkTreeItem : TSharedFromThis<FWorldBookmarkTreeItem>
 {
 public:
 	FAssetData BookmarkAsset;
+
+	bool bInEditingMode = false;
+
+	/** Delegate for when the context menu requests a rename */
+	DECLARE_DELEGATE(FOnRenameRequested);
+	FOnRenameRequested OnRenameRequested;
 };
 
 typedef TSharedPtr<FWorldBookmarkTreeItem> FWorldBookmarkTreeItemPtr;
@@ -32,6 +38,10 @@ private:
 	virtual TSharedRef<SWidget> GenerateWidgetForColumn(const FName& ColumnName) override;
 
 	FText GetLastAccessedText() const;
+	bool OnLabelTextVerifyChanged(const FText& InNewText, FText& OutErrorMessage);
+	void OnLabelTextCommitted(const FText& NewText, ETextCommit::Type InTextCommit);
+	void OnEnterEditingMode();
+	void OnExitEditingMode();
 
 private:
 	TWeakPtr<FWorldBookmarkTreeItem> WeakWorldBookmarkTreeItem;
@@ -58,6 +68,11 @@ private:
 	void OnSearchBoxTextChanged(const FText& InSearchText);
 	void OnColumnSortChanged(const EColumnSortPriority::Type SortPriority, const FName& ColumnId, const EColumnSortMode::Type InSortMode);
 	TSharedPtr<SWidget> OnContextMenuOpening();
+	void ItemScrolledIntoView(FWorldBookmarkTreeItemPtr WorldBookmarkTreeItem, const TSharedPtr<ITableRow>& Widget);
+
+	/** SWidget interface */
+	virtual void Tick(const FGeometry& AllottedGeometry, const double InCurrentTime, const float InDeltaTime) override;
+	virtual FReply OnKeyDown(const FGeometry& MyGeometry, const FKeyEvent& InKeyEvent) override;
 
 	void RefreshItems();
 	void SortItems();
@@ -76,6 +91,7 @@ private:
 
 	void OnMapChanged(UWorld* World, EMapChangeType MapChangeType);
 
+	FWorldBookmarkTreeItemPtr GetSelectedBookmarkTreeItem() const;
 	UWorldBookmark* GetSelectedBookmark() const;
 	void SetSelectedBookmark(UWorldBookmark* InBookmark);
 
@@ -89,6 +105,9 @@ private:
 	void AddSelectedBookmarkToFavorites();
 	void RemoveSelectedBookmarkFromFavorites();	
 
+	void RequestRenameSelectedBookmark();
+	void DeleteSelectedBookmark();
+
 private:
 	FDelegateHandle OnAssetAddedHandle;
 	FDelegateHandle OnAssetRemovedHandle;
@@ -109,4 +128,9 @@ private:
 	TPair<FName, EColumnSortMode::Type> ColumnsSortingParams[EColumnSortPriority::Max];
 	
 	FText CurrentSearchString;
+	
+	FWorldBookmarkTreeItemPtr BookmarkPendingRename;
+
+	bool bIsInTick;
+	bool bPendingRefresh;
 };
\ No newline at end of file
diff --git a/Engine/Source/Editor/WorldBookmark/WorldBookmark.Build.cs b/Engine/Source/Editor/WorldBookmark/WorldBookmark.Build.cs
--- a/Engine/Source/Editor/WorldBookmark/WorldBookmark.Build.cs
+++ b/Engine/Source/Editor/WorldBookmark/WorldBookmark.Build.cs
@@ -10,6 +10,7 @@ public class WorldBookmark : ModuleRules
 			new string[] {
 				"AssetDefinition",
 				"AssetRegistry",
+				"AssetTools",
 				"ContentBrowser",
 				"ContentBrowserData",
 				"Core",
