World Bookmark - Fixed CIS warning caused by duplicated LOC key

diff --git a/Engine/Source/Editor/WorldBookmark/Private/WorldBookmark/WorldBookmark.cpp b/Engine/Source/Editor/WorldBookmark/Private/WorldBookmark/WorldBookmark.cpp
--- a/Engine/Source/Editor/WorldBookmark/Private/WorldBookmark/WorldBookmark.cpp
+++ b/Engine/Source/Editor/WorldBookmark/Private/WorldBookmark/WorldBookmark.cpp
@@ -175,7 +175,7 @@ bool UWorldBookmark::CanUpdate(FText* FailureReason) const
 	{
 		if (FailureReason)
 		{
-			*FailureReason = LOCTEXT("CanUpdateFailure_InvalidWorld", "Unsaved world");
+			*FailureReason = LOCTEXT("CanUpdateFailure_UnsavedWorld", "Unsaved world");
 		}
 
 		return false;
