World Bookmark - Added settings to the World Bookmark browser
* Show Only Bookmarks For Current World
* Show Only Uncontrolled Bookmarks
* Show Only Favorite Bookmarks
* Show Only Last Recently Used Bookmarks (+ option to select number of LRU items)
* These are user settings saved locally

diff --git a/Engine/Source/Editor/WorldBookmark/Private/WorldBookmark/Browser/SWorldBookmarkBrowser.cpp b/Engine/Source/Editor/WorldBookmark/Private/WorldBookmark/Browser/SWorldBookmarkBrowser.cpp
--- a/Engine/Source/Editor/WorldBookmark/Private/WorldBookmark/Browser/SWorldBookmarkBrowser.cpp
+++ b/Engine/Source/Editor/WorldBookmark/Private/WorldBookmark/Browser/SWorldBookmarkBrowser.cpp
@@ -1,6 +1,7 @@
 // Copyright Epic Games, Inc. All Rights Reserved.
 #include "WorldBookmark/Browser/SWorldBookmarkBrowser.h"
 
+#include "Algo/AnyOf.h"
 #include "AssetRegistry/IAssetRegistry.h"
 #include "AssetRegistry/AssetData.h"
 #include "AssetRegistry/AssetRegistryHelpers.h"
@@ -15,13 +16,18 @@
 #include "Modules/ModuleManager.h"
 #include "ObjectTools.h"
 #include "PropertyEditorModule.h"
+#include "SourceControlHelpers.h"
 #include "SSimpleButton.h"
 #include "Toolkits/GlobalEditorCommonCommands.h"
+#include "UncontrolledChangelistsModule.h"
 #include "UObject/AssetRegistryTagsContext.h"
+#include "Widgets/Input/SComboButton.h"
 #include "Widgets/Input/SSearchBox.h"
+#include "Widgets/Input/SSpinBox.h"
 #include "Widgets/Text/SInlineEditableTextBlock.h"
 #include "WorldBookmark/WorldBookmark.h"
 #include "WorldBookmark/WorldBookmarkFactory.h"
+#include "WorldBookmark/Browser/WorldBookmarkBrowserSettings.h"
 
 #define LOCTEXT_NAMESPACE "WorldBookmarkBrowser"
 
@@ -82,6 +88,8 @@ SWorldBookmarkBrowser::SWorldBookmarkBrowser()
 
 	ColumnsSortingParams[EColumnSortPriority::Primary] = { WorldBookmarkBrowserColumn::Favorite::Id, EColumnSortMode::Ascending };
 	ColumnsSortingParams[EColumnSortPriority::Secondary] = { WorldBookmarkBrowserColumn::Label::Id, EColumnSortMode::Ascending };
+
+	OnSettingsChangedHandle = UWorldBookmarkBrowserSettings::OnSettingsChanged().AddRaw(this, &SWorldBookmarkBrowser::OnSettingsChanged);
 }
 
 SWorldBookmarkBrowser::~SWorldBookmarkBrowser()
@@ -97,6 +105,8 @@ SWorldBookmarkBrowser::~SWorldBookmarkBrowser()
 
 	FLevelEditorModule& LevelEditorModule = FModuleManager::LoadModuleChecked<FLevelEditorModule>(TEXT("LevelEditor"));
 	LevelEditorModule.OnMapChanged().Remove(OnMapChangedHandle);
+
+	UWorldBookmarkBrowserSettings::OnSettingsChanged().Remove(OnSettingsChangedHandle);
 }
 
 static FSlateIcon GetFavoriteWorldBookmarkIcon(bool bIsFavorite)
@@ -181,9 +191,9 @@ void SWorldBookmarkBrowser::Construct(const FArguments& InArgs)
 
 			+ SVerticalBox::Slot()
 			.AutoHeight()
+			.Padding(8.0f, 8.0f, 8.0f, 4.0f)
 			.VAlign(VAlign_Top)
 			.HAlign(HAlign_Fill)			
-			.Padding(0.0f)
 			[
 				SNew(SHorizontalBox)
 				
@@ -196,7 +206,6 @@ void SWorldBookmarkBrowser::Construct(const FArguments& InArgs)
 
 				+ SHorizontalBox::Slot()
 				.AutoWidth()
-				.HAlign(HAlign_Center)
 				.VAlign(VAlign_Center)
 				[
 					SNew(SButton)
@@ -209,6 +218,22 @@ void SWorldBookmarkBrowser::Construct(const FArguments& InArgs)
 						.Image(FAppStyle::Get().GetBrush("WorldPartition.WorldBookmark.Create"))
 					]
 				]
+
+				+ SHorizontalBox::Slot()
+				.AutoWidth()
+				.VAlign(VAlign_Center)
+				[
+					SNew(SComboButton)
+					.ComboButtonStyle(FAppStyle::Get(), "SimpleComboButtonWithIcon")
+					.OnGetMenuContent(this, &SWorldBookmarkBrowser::GetSettingsMenuContent)
+					.HasDownArrow(false)
+					.ButtonContent()
+					[
+						SNew(SImage)
+						.ColorAndOpacity(FSlateColor::UseForeground())
+						.Image(FAppStyle::Get().GetBrush("Icons.Settings"))
+					]
+				]
 			]
 
 			+ SVerticalBox::Slot()
@@ -244,6 +269,85 @@ void SWorldBookmarkBrowser::Construct(const FArguments& InArgs)
 		FCanExecuteAction::CreateSP(this, &SWorldBookmarkBrowser::IsValidBookmarkSelected));
 }
 
+TSharedRef<SWidget> SWorldBookmarkBrowser::GetSettingsMenuContent()
+{
+	FMenuBuilder MenuBuilder(false, nullptr);
+
+	// World Bookmark
+	MenuBuilder.BeginSection(TEXT("SettingsMenuSection_Show"), LOCTEXT("SettingsMenuSection_Show", "Show"));
+	{
+		// Show Only For Current World
+		MenuBuilder.AddMenuEntry(
+			LOCTEXT("Settings_ShowOnlyForCurrentWorld", "Show Only For Current World"),
+			LOCTEXT("Settings_ShowOnlyForCurrentWorld_Tooltip", "Displays only bookmarks that are bound to the current world"),
+			FSlateIcon(),
+			FUIAction(
+				FExecuteAction::CreateStatic(&UWorldBookmarkBrowserSettings::ToggleShowOnlyBookmarksForCurrentWorld),
+				FCanExecuteAction(),
+				FIsActionChecked::CreateStatic(&UWorldBookmarkBrowserSettings::GetShowOnlyBookmarksForCurrentWorld)
+			),
+			NAME_None,
+			EUserInterfaceActionType::ToggleButton
+		);
+
+		// Show Only Uncontrolled Bookmarks
+		MenuBuilder.AddMenuEntry(
+			LOCTEXT("Settings_ShowOnlyUncontrolledBookmarks", "Show Only Uncontrolled Bookmarks"),
+			LOCTEXT("Settings_ShowOnlyUncontrolledBookmarks_Tooltip", "Displays only local bookmarks that are kept in an uncontrolled changelist"),
+			FSlateIcon(),
+			FUIAction(
+				FExecuteAction::CreateStatic(&UWorldBookmarkBrowserSettings::ToggleShowOnlyUncontrolledBookmarks),
+				FCanExecuteAction(),
+				FIsActionChecked::CreateStatic(&UWorldBookmarkBrowserSettings::GetShowOnlyUncontrolledBookmarks)
+			),
+			NAME_None,
+			EUserInterfaceActionType::ToggleButton
+		);
+
+		// Show Only Favorite Bookmarks
+		MenuBuilder.AddMenuEntry(
+			LOCTEXT("Settings_ShowOnlyFavoriteBookmarks", "Show Only Favorite Bookmarks"),
+			LOCTEXT("Settings_ShowOnlyFavoriteBookmarks_Tooltip", "Displays only bookmarks that were flagged as favorite"),
+			FSlateIcon(),
+			FUIAction(
+				FExecuteAction::CreateStatic(&UWorldBookmarkBrowserSettings::ToggleShowOnlyFavoriteBookmarks),
+				FCanExecuteAction(),
+				FIsActionChecked::CreateStatic(&UWorldBookmarkBrowserSettings::GetShowOnlyFavoriteBookmarks)
+			),
+			NAME_None,
+			EUserInterfaceActionType::ToggleButton
+		);
+
+		// Show Only Last Recently Used Bookmarks
+		MenuBuilder.AddMenuEntry(
+			LOCTEXT("Settings_ShowOnlyLastRecentlyUsedBookmarks", "Show Only Last Recently Used Bookmarks"),
+			LOCTEXT("Settings_ShowOnlyLastRecentlyUsedBookmarks_Tooltip", "Displays only bookmarks that were recently used"),
+			FSlateIcon(),
+			FUIAction(
+				FExecuteAction::CreateStatic(&UWorldBookmarkBrowserSettings::ToggleShowOnlyLastRecentlyUsedBookmarks),
+				FCanExecuteAction(),
+				FIsActionChecked::CreateStatic(&UWorldBookmarkBrowserSettings::GetShowOnlyLastRecentlyUsedBookmarks)
+			),
+			NAME_None,
+			EUserInterfaceActionType::ToggleButton
+		);
+
+		// Last recently used items
+		TSharedRef<SSpinBox<int32>> Slider = SNew(SSpinBox<int32>)
+			.MinValue(0)
+			.MaxValue(10)
+			.ToolTipText(LOCTEXT("Settings_MaxLastRecentlyUsedItems_Tooltip", "Maximum number of last recently used items to display"))
+			.Value_Static(&UWorldBookmarkBrowserSettings::GetMaxLastRecentlyUsedItems)
+			.OnValueChanged_Static(UWorldBookmarkBrowserSettings::SetMaxLastRecentlyUsedItems);
+		
+		const bool bNoIndent = true;
+		MenuBuilder.AddWidget(Slider, LOCTEXT("Settings_MaxLastRecentlyUsedItems", "Last recently used items"), bNoIndent);
+	}
+	MenuBuilder.EndSection();
+
+	return MenuBuilder.MakeWidget();
+}
+
 FReply SWorldBookmarkBrowser::OnKeyDown(const FGeometry& MyGeometry, const FKeyEvent& InKeyEvent)
 {
 	bool bIsRenamingAsset = false;
@@ -277,30 +381,140 @@ void SWorldBookmarkBrowser::RefreshItems()
 
 	WorldBookmarkTreeItems.Reset();
 
-	for (const FAssetData& AssetData : AssetsData)
+	FSoftObjectPath CurrentWorldPath = FSoftObjectPath(GetCurrentWorld());
+	
+	TArray<FUncontrolledChangelistStateRef> UncontrolledChangelistsStates;
+	if (UWorldBookmarkBrowserSettings::GetShowOnlyUncontrolledBookmarks())
+	{
+		UncontrolledChangelistsStates = FUncontrolledChangelistsModule::Get().GetChangelistStates();
+	}
+
+	// Filter based on "Show Only Bookmarks For Current World"
+	auto FilterForCurrentWorld = [&CurrentWorldPath](const FAssetData& AssetData)
+	{
+		if (UWorldBookmarkBrowserSettings::GetShowOnlyBookmarksForCurrentWorld())
+		{
+			static const FName NAME_BookmarkWorld(TEXT("WorldName"));
+
+			FString WorldAssetName;
+			if (!AssetData.GetTagValue(NAME_BookmarkWorld, WorldAssetName))
+			{
+				return false;
+			}
+
+			FSoftObjectPath WorldAssetPath = FSoftObjectPath(WorldAssetName);
+			UAssetRegistryHelpers::FixupRedirectedAssetPath(WorldAssetPath);
+			if (CurrentWorldPath != WorldAssetPath)
+			{
+				return false;
+			}
+		}
+
+		return true;
+	};
+
+	// Filter based on "Show Only Uncontrolled Bookmarks"
+	auto FilterUncontrolledBookmarks = [&UncontrolledChangelistsStates](const FAssetData& AssetData)
 	{
-		static const FName NAME_BookmarkWorld(TEXT("WorldName"));
+		if (UWorldBookmarkBrowserSettings::GetShowOnlyUncontrolledBookmarks())
+		{
+			const FString BookmarkFileName = USourceControlHelpers::PackageFilename(AssetData.PackageName.ToString());
+			bool bIsUncontrolledFile = false;
+
+			for (const FUncontrolledChangelistStateRef& UncontrolledChangelistState : UncontrolledChangelistsStates)
+			{
+				if (UncontrolledChangelistState->GetFilenames().Contains(BookmarkFileName))
+				{
+					bIsUncontrolledFile = true;
+					break;
+				}
+			}
+
+			if (!bIsUncontrolledFile)
+			{
+				return false;
+			}
+		}
 
-		FString WorldAssetName;
-		if (!AssetData.GetTagValue(NAME_BookmarkWorld, WorldAssetName))
+		return true;
+	};
+
+	// Filter based on "Show Only Favorite Bookmarks"
+	auto FilterFavoriteBookmarks = [](const FAssetData& AssetData)
+	{
+		if (UWorldBookmarkBrowserSettings::GetShowOnlyFavoriteBookmarks())
 		{
-			continue;
+			UWorldBookmark* WorldBookmark = Cast<UWorldBookmark>(AssetData.GetAsset());
+			if (!WorldBookmark->GetIsUserFavorite())
+			{
+				return false;
+			}
 		}
 
-		FSoftObjectPath WorldAssetPath = FSoftObjectPath(WorldAssetName);
-		UAssetRegistryHelpers::FixupRedirectedAssetPath(WorldAssetPath);
+		return true;
+	};
 
-		FSoftObjectPath CurrentWorldPath = FSoftObjectPath(GetCurrentWorld());
-		if (CurrentWorldPath != WorldAssetPath)
+	// Filter based on "Show Only Last Recently Used Bookmarks"
+	auto FilterLRUBookmarks = [](const FAssetData& AssetData)
+	{
+		if (UWorldBookmarkBrowserSettings::GetShowOnlyLastRecentlyUsedBookmarks())
 		{
-			continue;
+			UWorldBookmark* WorldBookmark = Cast<UWorldBookmark>(AssetData.GetAsset());
+			if (WorldBookmark->GetUserLastLoadedTimeStampUTC() == FDateTime::MinValue())
+			{
+				return false;
+			}
 		}
 
+		return true;
+	};
+
+	// Filter based on the seach box value
+	auto FilterSearch = [this](const FAssetData& AssetData)
+	{
 		if (!CurrentSearchString.IsEmpty() && !AssetData.AssetName.ToString().Contains(CurrentSearchString.ToString()))
 		{
-			continue;
+			return false;
 		}
 
+		return true;
+	};
+		
+	using FAssetDataFilterPredicate = TFunction<bool(const FAssetData&)>;
+	TArray<FAssetDataFilterPredicate> AssetFiltering =
+	{
+		FilterForCurrentWorld,
+		FilterUncontrolledBookmarks,
+		FilterFavoriteBookmarks,
+		FilterLRUBookmarks,
+		FilterSearch
+	};
+
+	// Remove all elements which doesn't pass the filters
+	AssetsData.RemoveAllSwap([&AssetFiltering](const FAssetData& AssetData)
+	{
+		// Remove the element if any of the filter return false.
+		return Algo::AnyOf(AssetFiltering, [&AssetData](const FAssetDataFilterPredicate& Pred)
+		{
+			return !Pred(AssetData);
+		});
+	});
+
+	// Handle special case for LRU items, only keep N items
+	if (UWorldBookmarkBrowserSettings::GetShowOnlyLastRecentlyUsedBookmarks() &&
+		AssetsData.Num() > UWorldBookmarkBrowserSettings::GetMaxLastRecentlyUsedItems())
+	{
+		AssetsData.Sort([](const FAssetData& AssetA, const FAssetData& AssetB) 
+		{
+			UWorldBookmark* WorldBookmarkA = Cast<UWorldBookmark>(AssetA.GetAsset());
+			UWorldBookmark* WorldBookmarkB = Cast<UWorldBookmark>(AssetB.GetAsset());
+			return WorldBookmarkA->GetUserLastLoadedTimeStampUTC() > WorldBookmarkB->GetUserLastLoadedTimeStampUTC();
+		});
+		AssetsData.SetNum(UWorldBookmarkBrowserSettings::GetMaxLastRecentlyUsedItems());
+	}
+
+	for (const FAssetData& AssetData : AssetsData)
+	{
 		FWorldBookmarkTreeItemPtr& Item = WorldBookmarkTreeItems.Emplace_GetRef(MakeShared<FWorldBookmarkTreeItem>());
 		Item->BookmarkAsset = AssetData;
 	}
@@ -800,6 +1014,29 @@ void SWorldBookmarkBrowser::OnMapChanged(UWorld* World, EMapChangeType MapChange
 	RefreshItems();
 }
 
+void SWorldBookmarkBrowser::OnSettingsChanged()
+{
+	if (UWorldBookmarkBrowserSettings::GetShowOnlyUncontrolledBookmarks())
+	{
+		if (!OnUncontrolledChangelistModuleChangedHandle.IsValid())
+		{
+			OnUncontrolledChangelistModuleChangedHandle = FUncontrolledChangelistsModule::Get().OnUncontrolledChangelistModuleChanged.AddRaw(this, &SWorldBookmarkBrowser::OnUncontrolledChangelistModuleChanged);
+		}
+	}
+	else if (OnUncontrolledChangelistModuleChangedHandle.IsValid())
+	{
+		FUncontrolledChangelistsModule::Get().OnUncontrolledChangelistModuleChanged.Remove(OnUncontrolledChangelistModuleChangedHandle);
+		OnUncontrolledChangelistModuleChangedHandle.Reset();
+	}
+
+	RefreshItems();
+}
+
+void SWorldBookmarkBrowser::OnUncontrolledChangelistModuleChanged()
+{
+	RefreshItems();
+}
+
 void SWorldBookmarkBrowser::LoadSelectedBookmark()
 {
 	if (UWorldBookmark* SelectedBookmark = GetSelectedBookmark())
@@ -830,7 +1067,8 @@ void SWorldBookmarkBrowser::AddSelectedBookmarkToFavorites()
 {
 	if (UWorldBookmark* SelectedBookmark = GetSelectedBookmark())
 	{
-		return SelectedBookmark->SetIsUserFavorite(true);
+		SelectedBookmark->SetIsUserFavorite(true);
+		RefreshItems();
 	}
 }
 
@@ -838,7 +1076,8 @@ void SWorldBookmarkBrowser::RemoveSelectedBookmarkFromFavorites()
 {
 	if (UWorldBookmark* SelectedBookmark = GetSelectedBookmark())
 	{
-		return SelectedBookmark->SetIsUserFavorite(false);
+		SelectedBookmark->SetIsUserFavorite(false);
+		RefreshItems();
 	}
 }
 
diff --git a/Engine/Source/Editor/WorldBookmark/Private/WorldBookmark/Browser/SWorldBookmarkBrowser.h b/Engine/Source/Editor/WorldBookmark/Private/WorldBookmark/Browser/SWorldBookmarkBrowser.h
--- a/Engine/Source/Editor/WorldBookmark/Private/WorldBookmark/Browser/SWorldBookmarkBrowser.h
+++ b/Engine/Source/Editor/WorldBookmark/Private/WorldBookmark/Browser/SWorldBookmarkBrowser.h
@@ -1,14 +1,14 @@
 // Copyright Epic Games, Inc. All Rights Reserved.
 #pragma once
 
+#include "AssetTextFilter.h"
+#include "ContentBrowserItem.h"
 #include "Misc/TextFilter.h"
 #include "Templates/SharedPointer.h"
-#include "Widgets/Views/SHeaderRow.h"	
-#include "Widgets/Views/STreeView.h"
-#include "ContentBrowserItem.h"
-#include "AssetTextFilter.h"
 #include "UObject/ObjectSaveContext.h"
 #include "UnrealEdMisc.h"
+#include "Widgets/Views/SHeaderRow.h"	
+#include "Widgets/Views/STreeView.h"
 
 class UWorldBookmark;
 class IDetailsView;
@@ -68,6 +68,7 @@ private:
 	void OnSearchBoxTextChanged(const FText& InSearchText);
 	void OnColumnSortChanged(const EColumnSortPriority::Type SortPriority, const FName& ColumnId, const EColumnSortMode::Type InSortMode);
 	TSharedPtr<SWidget> OnContextMenuOpening();
+	TSharedRef<SWidget> GetSettingsMenuContent();
 	void ItemScrolledIntoView(FWorldBookmarkTreeItemPtr WorldBookmarkTreeItem, const TSharedPtr<ITableRow>& Widget);
 
 	/** SWidget interface */
@@ -90,6 +91,10 @@ private:
 	void OnObjectPreSave(UObject* InObject, FObjectPreSaveContext InObjectPreSaveContext);
 
 	void OnMapChanged(UWorld* World, EMapChangeType MapChangeType);
+	
+	void OnSettingsChanged();
+
+	void OnUncontrolledChangelistModuleChanged();
 
 	FWorldBookmarkTreeItemPtr GetSelectedBookmarkTreeItem() const;
 	UWorldBookmark* GetSelectedBookmark() const;
@@ -116,6 +121,8 @@ private:
 	FDelegateHandle OnAssetUpdatedOnDiskHandle;
 	FDelegateHandle OnObjectPropertyChangedHandle;
 	FDelegateHandle OnMapChangedHandle;
+	FDelegateHandle OnSettingsChangedHandle;
+	FDelegateHandle OnUncontrolledChangelistModuleChangedHandle;
 
 	TArray<FWorldBookmarkTreeItemPtr>	WorldBookmarkTreeItems;
 
diff --git a/Engine/Source/Editor/WorldBookmark/Private/WorldBookmark/Browser/WorldBookmarkBrowserSettings.cpp b/Engine/Source/Editor/WorldBookmark/Private/WorldBookmark/Browser/WorldBookmarkBrowserSettings.cpp
new file mode 100644
--- /dev/null
+++ b/Engine/Source/Editor/WorldBookmark/Private/WorldBookmark/Browser/WorldBookmarkBrowserSettings.cpp
@@ -0,0 +1,92 @@
+// Copyright Epic Games, Inc. All Rights Reserved.
+
+#include "WorldBookmarkBrowserSettings.h"
+
+TObjectPtr<UWorldBookmarkBrowserSettings> UWorldBookmarkBrowserSettings::Instance = nullptr;
+
+UWorldBookmarkBrowserSettings& UWorldBookmarkBrowserSettings::Get()
+{
+	if (!Instance)
+	{
+		Instance = NewObject<UWorldBookmarkBrowserSettings>();
+		Instance->LoadEditorConfig();
+		Instance->AddToRoot();
+	}
+
+	return *Instance;
+}
+
+void UWorldBookmarkBrowserSettings::ApplySettingsChanges()
+{
+	SaveEditorConfig();
+	OnSettingsChangedDelegate.Broadcast();
+}
+
+UWorldBookmarkBrowserSettings::FWorldBookmarkBrowserSettingChanged& UWorldBookmarkBrowserSettings::OnSettingsChanged()
+{
+	UWorldBookmarkBrowserSettings& Settings = Get();
+	return Settings.OnSettingsChangedDelegate;
+}
+
+bool UWorldBookmarkBrowserSettings::GetShowOnlyBookmarksForCurrentWorld()
+{
+	return Get().bShowOnlyBookmarksForCurrentWorld;
+}
+
+void UWorldBookmarkBrowserSettings::ToggleShowOnlyBookmarksForCurrentWorld()
+{
+	UWorldBookmarkBrowserSettings& Settings = Get();
+	Settings.bShowOnlyBookmarksForCurrentWorld = !Settings.bShowOnlyBookmarksForCurrentWorld;
+	Settings.ApplySettingsChanges();
+}
+
+bool UWorldBookmarkBrowserSettings::GetShowOnlyUncontrolledBookmarks()
+{
+	return Get().bShowOnlyUncontrolledBookmarks;
+}
+
+void UWorldBookmarkBrowserSettings::ToggleShowOnlyUncontrolledBookmarks()
+{
+	UWorldBookmarkBrowserSettings& Settings = Get();
+	Settings.bShowOnlyUncontrolledBookmarks = !Settings.bShowOnlyUncontrolledBookmarks;
+	Settings.ApplySettingsChanges();
+}
+
+bool UWorldBookmarkBrowserSettings::GetShowOnlyFavoriteBookmarks()
+{
+	return Get().bShowOnlyFavoriteBookmarks;
+}
+
+void UWorldBookmarkBrowserSettings::ToggleShowOnlyFavoriteBookmarks()
+{
+	UWorldBookmarkBrowserSettings& Settings = Get();
+	Settings.bShowOnlyFavoriteBookmarks = !Settings.bShowOnlyFavoriteBookmarks;
+	Settings.ApplySettingsChanges();
+}
+
+bool UWorldBookmarkBrowserSettings::GetShowOnlyLastRecentlyUsedBookmarks()
+{
+	return Get().bShowOnlyLastRecentlyUsedBookmarks;
+}
+
+void UWorldBookmarkBrowserSettings::ToggleShowOnlyLastRecentlyUsedBookmarks()
+{
+	UWorldBookmarkBrowserSettings& Settings = Get();
+	Settings.bShowOnlyLastRecentlyUsedBookmarks = !Settings.bShowOnlyLastRecentlyUsedBookmarks;
+	Settings.ApplySettingsChanges();
+}
+
+int32 UWorldBookmarkBrowserSettings::GetMaxLastRecentlyUsedItems()
+{
+	return Get().MaxLastRecentlyUsedItems;
+}
+
+void UWorldBookmarkBrowserSettings::SetMaxLastRecentlyUsedItems(int32 InMaxLastRecentlyUsedItems)
+{
+	UWorldBookmarkBrowserSettings& Settings = Get();
+	if (Settings.MaxLastRecentlyUsedItems != InMaxLastRecentlyUsedItems)
+	{
+		Settings.MaxLastRecentlyUsedItems = InMaxLastRecentlyUsedItems;
+		Settings.ApplySettingsChanges();
+	}
+}
\ No newline at end of file
diff --git a/Engine/Source/Editor/WorldBookmark/Private/WorldBookmark/Browser/WorldBookmarkBrowserSettings.h b/Engine/Source/Editor/WorldBookmark/Private/WorldBookmark/Browser/WorldBookmarkBrowserSettings.h
new file mode 100644
--- /dev/null
+++ b/Engine/Source/Editor/WorldBookmark/Private/WorldBookmark/Browser/WorldBookmarkBrowserSettings.h
@@ -0,0 +1,56 @@
+// Copyright Epic Games, Inc. All Rights Reserved.
+#pragma once
+
+#include "EditorConfigBase.h"
+
+#include "WorldBookmarkBrowserSettings.generated.h"
+
+UCLASS(EditorConfig = "WorldBookmarkBrowser")
+class UWorldBookmarkBrowserSettings : public UEditorConfigBase
+{
+	GENERATED_BODY()
+
+public:
+	DECLARE_MULTICAST_DELEGATE(FWorldBookmarkBrowserSettingChanged)
+
+	static FWorldBookmarkBrowserSettingChanged& OnSettingsChanged();
+
+	static bool GetShowOnlyBookmarksForCurrentWorld();
+	static void ToggleShowOnlyBookmarksForCurrentWorld();
+
+	static bool GetShowOnlyUncontrolledBookmarks();
+	static void ToggleShowOnlyUncontrolledBookmarks();
+
+	static bool GetShowOnlyFavoriteBookmarks();
+	static void ToggleShowOnlyFavoriteBookmarks();
+
+	static bool GetShowOnlyLastRecentlyUsedBookmarks();
+	static void ToggleShowOnlyLastRecentlyUsedBookmarks();
+
+	static int32 GetMaxLastRecentlyUsedItems();
+	static void SetMaxLastRecentlyUsedItems(int32 InMaxLastRecentlyUsedItems);
+
+private:
+	static UWorldBookmarkBrowserSettings& Get();
+
+	void ApplySettingsChanges();
+
+	UPROPERTY(meta=(EditorConfig))
+	bool bShowOnlyBookmarksForCurrentWorld = true;
+
+	UPROPERTY(meta=(EditorConfig))
+	bool bShowOnlyUncontrolledBookmarks = false;
+
+	UPROPERTY(meta=(EditorConfig))
+	bool bShowOnlyFavoriteBookmarks = false;
+
+	UPROPERTY(meta=(EditorConfig))
+	bool bShowOnlyLastRecentlyUsedBookmarks = false;
+
+	UPROPERTY(meta = (EditorConfig))
+	int32 MaxLastRecentlyUsedItems = 5;
+
+	FWorldBookmarkBrowserSettingChanged OnSettingsChangedDelegate;
+
+	static TObjectPtr<UWorldBookmarkBrowserSettings> Instance;
+};
diff --git a/Engine/Source/Editor/WorldBookmark/WorldBookmark.Build.cs b/Engine/Source/Editor/WorldBookmark/WorldBookmark.Build.cs
--- a/Engine/Source/Editor/WorldBookmark/WorldBookmark.Build.cs
+++ b/Engine/Source/Editor/WorldBookmark/WorldBookmark.Build.cs
@@ -16,6 +16,7 @@ public class WorldBookmark : ModuleRules
 				"Core",
 				"CoreUObject",
 				"EditorSubsystem",
+				"EditorConfig",
 				"Engine",
 				"InputCore",
 				"PropertyEditor",
@@ -25,6 +26,7 @@ public class WorldBookmark : ModuleRules
 				"SlateCore",				
 				"ToolMenus",
 				"ToolWidgets",
+				"UncontrolledChangelists",
 				"UnrealEd",
 				"WorkspaceMenuStructure"
 			}
