World Bookmark - Commands/style cleanup
* Moved the WorldBookmark style logic out of the global FStarshipEditorStyle

diff --git a/Engine/Source/Editor/EditorStyle/Private/StarshipStyle.cpp b/Engine/Source/Editor/EditorStyle/Private/StarshipStyle.cpp
--- a/Engine/Source/Editor/EditorStyle/Private/StarshipStyle.cpp
+++ b/Engine/Source/Editor/EditorStyle/Private/StarshipStyle.cpp
@@ -2062,16 +2062,6 @@ void FStarshipEditorStyle::FStyle::SetupWorldPartitionStyles()
 	Set("WorldPartition.ShowMiniMap", new IMAGE_BRUSH_SVG("Starship/AssetIcons/Texture2D_16", Icon16x16)); 
 	Set("WorldPartition.ShowCoords", new IMAGE_BRUSH_SVG("WorldPartition/Coordinate_16", Icon16x16));
 	Set("WorldPartition.ShowMouseCoords", new IMAGE_BRUSH_SVG("WorldPartition/MouseCoordinates_16", Icon16x16));
-
-	Set("WorldPartition.WorldBookmark", new IMAGE_BRUSH_SVG("Starship/AssetIcons/WorldBookmark_64", Icon64x64));
-	Set("WorldPartition.WorldBookmark.TabIcon", new IMAGE_BRUSH_SVG("Starship/AssetIcons/WorldBookmark_16", Icon16x16));
-	Set("WorldPartition.WorldBookmark.GoTo", new IMAGE_BRUSH_SVG("Starship/Common/NextArrow", Icon16x16));
-	Set("WorldPartition.WorldBookmark.Update", new CORE_IMAGE_BRUSH_SVG("Starship/Common/Update", Icon20x20));
-	Set("WorldPartition.WorldBookmark.Create", new CORE_IMAGE_BRUSH_SVG("Starship/Common/plus", Icon16x16, FStyleColors::AccentGreen));
-	Set("WorldPartition.WorldBookmark.IsFavorite", new IMAGE_BRUSH("Icons/Star_16x", Icon16x16));
-	Set("WorldPartition.WorldBookmark.IsNotFavorite", new IMAGE_BRUSH("Icons/EmptyStar_16x", Icon16x16));
-	Set("WorldPartition.WorldBookmark.RecentlyUsed", new IMAGE_BRUSH("Icons/Clock_24x", Icon16x16));
-
 	Set("WorldPartition.FocusSelection", new IMAGE_BRUSH_SVG("WorldPartition/FramingSelection_20", Icon16x16));
 	Set("WorldPartition.FocusLoadedRegions", new IMAGE_BRUSH_SVG("WorldPartition/FrameRegions_20", Icon16x16));
 	Set("WorldPartition.FocusWorld", new IMAGE_BRUSH_SVG("WorldPartition/FrameWorld_20", Icon16x16));
@@ -6421,8 +6411,7 @@ void FStarshipEditorStyle::FStyle::SetupClassIconsAndThumbnails()
 			{TEXT("WidgetBlueprintGeneratedClass")},
 			{TEXT("WindDirectionalSource")},
 			{TEXT("World")},
-			{TEXT("WorldDataLayers")},
-			{TEXT("WorldBookmark")},
+			{TEXT("WorldDataLayers")}
 		};
 	
 		// SVG Asset icons
diff --git a/Engine/Source/Editor/WorldBookmark/Private/WorldBookmark/AssetDefinition_WorldBookmark.cpp b/Engine/Source/Editor/WorldBookmark/Private/WorldBookmark/AssetDefinition_WorldBookmark.cpp
--- a/Engine/Source/Editor/WorldBookmark/Private/WorldBookmark/AssetDefinition_WorldBookmark.cpp
+++ b/Engine/Source/Editor/WorldBookmark/Private/WorldBookmark/AssetDefinition_WorldBookmark.cpp
@@ -2,6 +2,8 @@
 
 #include "WorldBookmark/AssetDefinition_WorldBookmark.h"
 #include "WorldBookmark/WorldBookmark.h"
+#include "WorldBookmark/WorldBookmarkCommands.h"
+#include "WorldBookmark/WorldBookmarkStyle.h"
 #include "ContentBrowserMenuContexts.h"
 #include "Toolkits/SimpleAssetEditor.h"
 #include "ToolMenus.h"
@@ -95,28 +97,20 @@ namespace MenuExtension_Bookmark
 			FToolMenuSection& Section = Menu->FindOrAddSection("GetAssetActions");
 			Section.AddDynamicEntry(NAME_None, FNewToolMenuSectionDelegate::CreateLambda([](FToolMenuSection& InSection)
 			{
-				const TAttribute<FText> Label = LOCTEXT("WorldBookmark_GoTo", "Go To");
-				const TAttribute<FText> ToolTip = LOCTEXT("WorldBookmark_GoToTooltip", "Will open the world and apply the editor state saved with this bookmark.");
-				const FSlateIcon Icon = FSlateIcon(FAppStyle::GetAppStyleSetName(), "WorldPartition.WorldBookmark.GoTo");
+				auto AddMenuEntry = [&InSection](TSharedPtr<FUICommandInfo> InCommandInfo, FToolUIAction InUIAction)
+				{
+					InSection.AddMenuEntry(InCommandInfo->GetCommandName(), InCommandInfo->GetLabel(), InCommandInfo->GetDescription(), InCommandInfo->GetIcon(), InUIAction);
+				};
 
 				FToolUIAction UIAction;
+
 				UIAction.ExecuteAction = FToolMenuExecuteAction::CreateStatic(&ExecuteGoToWorldBookmark);
 				UIAction.CanExecuteAction = FToolMenuCanExecuteAction::CreateStatic(&CanExecuteGoToWorldBookmark);
+				AddMenuEntry(FWorldBookmarkCommands::Get().LoadBookmark, UIAction);
 
-				InSection.AddMenuEntry("Bookmark_Edit", Label, ToolTip, Icon, UIAction);
-			}));
-
-			Section.AddDynamicEntry(NAME_None, FNewToolMenuSectionDelegate::CreateLambda([](FToolMenuSection& InSection)
-			{
-				const TAttribute<FText> Label = LOCTEXT("WorldBookmark_Update", "Update");
-				const TAttribute<FText> ToolTip = LOCTEXT("WorldBookmark_UpdateTooltip", "Update the bookmark with the current loaded state of the current loaded world.");
-				const FSlateIcon Icon = FSlateIcon(FAppStyle::GetAppStyleSetName(), "WorldPartition.WorldBookmark.Update");
-					
-				FToolUIAction UIAction;
 				UIAction.ExecuteAction = FToolMenuExecuteAction::CreateStatic(&ExecuteUpdateWorldBookmark);
 				UIAction.CanExecuteAction = FToolMenuCanExecuteAction::CreateStatic(&CanExecuteUpdateWorldBookmark);
-
-				InSection.AddMenuEntry("WorldBookmark_Update", Label, ToolTip, Icon, UIAction);
+				AddMenuEntry(FWorldBookmarkCommands::Get().UpdateBookmark, UIAction);
 			}));
 		}));
 	});
diff --git a/Engine/Source/Editor/WorldBookmark/Private/WorldBookmark/Browser/SWorldBookmarkBrowser.cpp b/Engine/Source/Editor/WorldBookmark/Private/WorldBookmark/Browser/SWorldBookmarkBrowser.cpp
--- a/Engine/Source/Editor/WorldBookmark/Private/WorldBookmark/Browser/SWorldBookmarkBrowser.cpp
+++ b/Engine/Source/Editor/WorldBookmark/Private/WorldBookmark/Browser/SWorldBookmarkBrowser.cpp
@@ -26,7 +26,9 @@
 #include "Widgets/Input/SSpinBox.h"
 #include "Widgets/Text/SInlineEditableTextBlock.h"
 #include "WorldBookmark/WorldBookmark.h"
+#include "WorldBookmark/WorldBookmarkCommands.h"
 #include "WorldBookmark/WorldBookmarkFactory.h"
+#include "WorldBookmark/WorldBookmarkStyle.h"
 #include "WorldBookmark/Browser/WorldBookmarkBrowserSettings.h"
 
 #define LOCTEXT_NAMESPACE "WorldBookmarkBrowser"
@@ -111,14 +113,14 @@ SWorldBookmarkBrowser::~SWorldBookmarkBrowser()
 
 static FSlateIcon GetFavoriteWorldBookmarkIcon(bool bIsFavorite)
 {
-	const FName FavoriteIconName = bIsFavorite ? "WorldPartition.WorldBookmark.IsFavorite" : "WorldPartition.WorldBookmark.IsNotFavorite";
-	return FSlateIcon(FAppStyle::Get().GetStyleSetName(), FavoriteIconName);
+	const FName FavoriteIconName = bIsFavorite ? "WorldBookmark.IsFavorite" : "WorldBookmark.IsNotFavorite";
+	return FSlateIcon(FWorldBookmarkStyle::Get().GetStyleSetName(), FavoriteIconName);
 }
 
 static FSlateIcon GetRecentlyUsedWorldBookmarkIcon()
 {
-	const FName RecentlyUsedIconName = "WorldPartition.WorldBookmark.RecentlyUsed";
-	return FSlateIcon(FAppStyle::Get().GetStyleSetName(), RecentlyUsedIconName);
+	const FName RecentlyUsedIconName = "WorldBookmark.RecentlyUsed";
+	return FSlateIcon(FWorldBookmarkStyle::Get().GetStyleSetName(), RecentlyUsedIconName);
 }
 
 void SWorldBookmarkBrowser::Construct(const FArguments& InArgs)
@@ -215,7 +217,7 @@ void SWorldBookmarkBrowser::Construct(const FArguments& InArgs)
 					[
 						SNew(SImage)
 						.ColorAndOpacity(FSlateColor::UseForeground())
-						.Image(FAppStyle::Get().GetBrush("WorldPartition.WorldBookmark.Create"))
+						.Image(FWorldBookmarkStyle::Get().GetBrush("WorldBookmark.Create"))
 					]
 				]
 
@@ -256,6 +258,7 @@ void SWorldBookmarkBrowser::Construct(const FArguments& InArgs)
 	];
 
 	Commands = MakeShareable(new FUICommandList());
+
 	Commands->MapAction(FGlobalEditorCommonCommands::Get().FindInContentBrowser,
 		FExecuteAction::CreateSP(this, &SWorldBookmarkBrowser::FindSelectedBookmarkInContentBrowser),
 		FCanExecuteAction::CreateSP(this, &SWorldBookmarkBrowser::IsValidBookmarkSelected));
@@ -267,6 +270,26 @@ void SWorldBookmarkBrowser::Construct(const FArguments& InArgs)
 	Commands->MapAction(FGenericCommands::Get().Delete,
 		FExecuteAction::CreateSP(this, &SWorldBookmarkBrowser::DeleteSelectedBookmark),
 		FCanExecuteAction::CreateSP(this, &SWorldBookmarkBrowser::IsValidBookmarkSelected));
+
+	Commands->MapAction(FWorldBookmarkCommands::Get().LoadBookmark,
+		FExecuteAction::CreateSP(this, &SWorldBookmarkBrowser::LoadSelectedBookmark),
+		FCanExecuteAction::CreateSP(this, &SWorldBookmarkBrowser::IsValidBookmarkSelected));
+
+	Commands->MapAction(FWorldBookmarkCommands::Get().UpdateBookmark,
+		FExecuteAction::CreateSP(this, &SWorldBookmarkBrowser::UpdateSelectedBookmark),
+		FCanExecuteAction::CreateSP(this, &SWorldBookmarkBrowser::IsValidBookmarkSelected));
+	
+	Commands->MapAction(FWorldBookmarkCommands::Get().AddToFavorite,
+		FExecuteAction::CreateSP(this, &SWorldBookmarkBrowser::AddSelectedBookmarkToFavorites),
+		FCanExecuteAction::CreateSP(this, &SWorldBookmarkBrowser::IsValidBookmarkSelected),
+		FIsActionChecked(),
+		FIsActionButtonVisible::CreateSP(this, &SWorldBookmarkBrowser::IsSelectedBookmarkNotFavorite));
+
+	Commands->MapAction(FWorldBookmarkCommands::Get().RemoveFromFavorite,
+		FExecuteAction::CreateSP(this, &SWorldBookmarkBrowser::RemoveSelectedBookmarkFromFavorites),
+		FCanExecuteAction::CreateSP(this, &SWorldBookmarkBrowser::IsValidBookmarkSelected),
+		FIsActionChecked(),
+		FIsActionButtonVisible::CreateSP(this, &SWorldBookmarkBrowser::IsSelectedBookmarkFavorite));
 }
 
 TSharedRef<SWidget> SWorldBookmarkBrowser::GetSettingsMenuContent()
@@ -605,7 +628,7 @@ TSharedRef<SWidget> FWorldBookmarkTableRow::GenerateWidgetForColumn(const FName&
 						.VAlign(VAlign_Center)
 						[
 							SNew(SImage)
-								.DesiredSizeOverride(FVector2D(24, 24))
+								.DesiredSizeOverride(FVector2D(20, 20))
 								.ColorAndOpacity(FSlateColor::UseForeground())
 								.Image(GetFavoriteWorldBookmarkIcon(true).GetIcon())
 						];
@@ -625,7 +648,7 @@ TSharedRef<SWidget> FWorldBookmarkTableRow::GenerateWidgetForColumn(const FName&
 						.VAlign(VAlign_Center)
 						[
 							SNew(SImage)
-								.DesiredSizeOverride(FVector2D(24, 24))
+								.DesiredSizeOverride(FVector2D(20, 20))
 								.ColorAndOpacity(FSlateColor::UseForeground())
 								.Image(GetRecentlyUsedWorldBookmarkIcon().GetIcon())
 						];
@@ -864,51 +887,10 @@ TSharedPtr<SWidget> SWorldBookmarkBrowser::OnContextMenuOpening()
 	// World Bookmark
 	MenuBuilder.BeginSection(TEXT("WorldBookmarkMenuSection"), LOCTEXT("WorldBookmarkMenuSection_Label", "World Bookmark"));
 	{
-		// Load
-		MenuBuilder.AddMenuEntry(
-			LOCTEXT("LoadBookmark", "Load Bookmark"),
-			LOCTEXT("LoadBookmarkToolTip", "Load this bookmark and apply it's state to the world."),
-			FSlateIcon(FAppStyle::Get().GetStyleSetName(), "WorldPartition.WorldBookmark.GoTo"),
-			FUIAction(
-				FExecuteAction::CreateSP(this, &SWorldBookmarkBrowser::LoadSelectedBookmark),
-				FCanExecuteAction::CreateSP(this, &SWorldBookmarkBrowser::IsValidBookmarkSelected)
-			)
-		);
-
-		// Update
-		MenuBuilder.AddMenuEntry(
-			LOCTEXT("UpdateBookmark", "Update Bookmark"),
-			LOCTEXT("UpdateBookmarkToolTip", "Update this bookmark using the current state of the world."),
-			FSlateIcon(FAppStyle::Get().GetStyleSetName(), "WorldPartition.WorldBookmark.Update"),
-			FUIAction(
-				FExecuteAction::CreateSP(this, &SWorldBookmarkBrowser::UpdateSelectedBookmark),
-				FCanExecuteAction::CreateSP(this, &SWorldBookmarkBrowser::IsValidBookmarkSelected)
-			)
-		);
-
-		MenuBuilder.AddMenuEntry(
-			LOCTEXT("AddToFavorite", "Add to Favorite"),
-			LOCTEXT("AddToFavorite_ToolTip", "Add this bookmark to your favorites."),
-			GetFavoriteWorldBookmarkIcon(false),
-			FUIAction(
-				FExecuteAction::CreateSP(this, &SWorldBookmarkBrowser::AddSelectedBookmarkToFavorites),
-				FCanExecuteAction::CreateSP(this, &SWorldBookmarkBrowser::IsValidBookmarkSelected),
-				FIsActionChecked(),
-				FIsActionButtonVisible::CreateSP(this, &SWorldBookmarkBrowser::IsSelectedBookmarkNotFavorite)
-			)
-		);
-
-		MenuBuilder.AddMenuEntry(
-			LOCTEXT("RemoveFromFavorite", "Remove from Favorite"),
-			LOCTEXT("RemoveFromFavorite_ToolTip", "Remove this bookmark from your favorites."),
-			GetFavoriteWorldBookmarkIcon(true),
-			FUIAction(
-				FExecuteAction::CreateSP(this, &SWorldBookmarkBrowser::RemoveSelectedBookmarkFromFavorites),
-				FCanExecuteAction::CreateSP(this, &SWorldBookmarkBrowser::IsValidBookmarkSelected),
-				FIsActionChecked(),
-				FIsActionButtonVisible::CreateSP(this, &SWorldBookmarkBrowser::IsSelectedBookmarkFavorite)
-			)
-		);
+		MenuBuilder.AddMenuEntry(FWorldBookmarkCommands::Get().LoadBookmark);
+		MenuBuilder.AddMenuEntry(FWorldBookmarkCommands::Get().UpdateBookmark);
+		MenuBuilder.AddMenuEntry(FWorldBookmarkCommands::Get().AddToFavorite);
+		MenuBuilder.AddMenuEntry(FWorldBookmarkCommands::Get().RemoveFromFavorite);
 	}
 
 	MenuBuilder.EndSection();
diff --git a/Engine/Source/Editor/WorldBookmark/Private/WorldBookmark/WorldBookmarkCommands.cpp b/Engine/Source/Editor/WorldBookmark/Private/WorldBookmark/WorldBookmarkCommands.cpp
new file mode 100644
--- /dev/null
+++ b/Engine/Source/Editor/WorldBookmark/Private/WorldBookmark/WorldBookmarkCommands.cpp
@@ -0,0 +1,25 @@
+// Copyright Epic Games, Inc. All Rights Reserved.
+#include "WorldBookmark/WorldBookmarkCommands.h"
+#include "WorldBookmark/WorldBookmarkStyle.h"
+
+#define LOCTEXT_NAMESPACE "WorldBookmarkCommands"
+
+FWorldBookmarkCommands::FWorldBookmarkCommands()
+	: TCommands<FWorldBookmarkCommands>
+	(
+		"WorldBookmark",
+		LOCTEXT("WorldBookmark", "World Bookmark"),
+		NAME_None,
+		FWorldBookmarkStyle::Get().GetStyleSetName()
+	)
+{}
+
+void FWorldBookmarkCommands::RegisterCommands()
+{
+	UI_COMMAND(LoadBookmark, "Load Bookmark", "Load this bookmark and apply it's state to the world.", EUserInterfaceActionType::Button, FInputChord());
+	UI_COMMAND(UpdateBookmark, "Update Bookmark", "Update this bookmark using the current state of the world.", EUserInterfaceActionType::Button, FInputChord());
+	UI_COMMAND(AddToFavorite, "Add to Favorite", "Add this bookmark to your favorites.", EUserInterfaceActionType::Button, FInputChord());
+	UI_COMMAND(RemoveFromFavorite, "Remove from Favorite", "Remove this bookmark from your favorites.", EUserInterfaceActionType::Button, FInputChord());
+}
+
+#undef LOCTEXT_NAMESPACE
diff --git a/Engine/Source/Editor/WorldBookmark/Private/WorldBookmark/WorldBookmarkCommands.h b/Engine/Source/Editor/WorldBookmark/Private/WorldBookmark/WorldBookmarkCommands.h
new file mode 100644
--- /dev/null
+++ b/Engine/Source/Editor/WorldBookmark/Private/WorldBookmark/WorldBookmarkCommands.h
@@ -0,0 +1,21 @@
+// Copyright Epic Games, Inc. All Rights Reserved.
+#pragma once
+
+#include "Framework/Commands/Commands.h"
+
+class FWorldBookmarkCommands : public TCommands<FWorldBookmarkCommands>
+{
+public:
+	FWorldBookmarkCommands();
+
+	// Context Menu
+	TSharedPtr<FUICommandInfo> LoadBookmark;
+	TSharedPtr<FUICommandInfo> UpdateBookmark;
+	TSharedPtr<FUICommandInfo> AddToFavorite;
+	TSharedPtr<FUICommandInfo> RemoveFromFavorite;
+		
+	/**
+	 * Initialize commands
+	 */
+	virtual void RegisterCommands() override;
+};
\ No newline at end of file
diff --git a/Engine/Source/Editor/WorldBookmark/Private/WorldBookmark/WorldBookmarkModule.cpp b/Engine/Source/Editor/WorldBookmark/Private/WorldBookmark/WorldBookmarkModule.cpp
--- a/Engine/Source/Editor/WorldBookmark/Private/WorldBookmark/WorldBookmarkModule.cpp
+++ b/Engine/Source/Editor/WorldBookmark/Private/WorldBookmark/WorldBookmarkModule.cpp
@@ -2,7 +2,9 @@
 
 #include "WorldBookmark/WorldBookmarkModule.h"
 #include "WorldBookmark/WorldBookmark.h"
+#include "WorldBookmark/WorldBookmarkCommands.h"
 #include "WorldBookmark/WorldBookmarkDetailsCustomization.h"
+#include "WorldBookmark/WorldBookmarkStyle.h"
 #include "WorldBookmark/Browser/SWorldBookmarkBrowser.h"
 
 #include "LevelEditor.h"
@@ -22,6 +24,8 @@ void FWorldBookmarkModule::StartupModule()
 {
 	if (GIsEditor && !IsRunningCommandlet())
 	{
+		FWorldBookmarkCommands::Register();
+
 		FLevelEditorModule& LevelEditorModule = FModuleManager::Get().LoadModuleChecked<FLevelEditorModule>("LevelEditor");
 		LevelEditorModule.OnRegisterTabs().AddRaw(this, &FWorldBookmarkModule::RegisterWorldBookmarkBrowserTab);
 
@@ -62,7 +66,7 @@ void FWorldBookmarkModule::RegisterWorldBookmarkBrowserTab(TSharedPtr<FTabManage
 {
 	const IWorkspaceMenuStructure& MenuStructure = WorkspaceMenu::GetMenuStructure();
 
-	const FSlateIcon WorldPartitionIcon(FAppStyle::GetAppStyleSetName(), "WorldPartition.WorldBookmark.TabIcon");
+	const FSlateIcon WorldPartitionIcon(FWorldBookmarkStyle::Get().GetStyleSetName(), "WorldBookmark.TabIcon");
 
 	InTabManager->RegisterTabSpawner(WorldBookmarkBrowserTabId,
 		FOnSpawnTab::CreateRaw(this, &FWorldBookmarkModule::SpawnWorldBookmarkBrowserTab))
diff --git a/Engine/Source/Editor/WorldBookmark/Private/WorldBookmark/WorldBookmarkStyle.cpp b/Engine/Source/Editor/WorldBookmark/Private/WorldBookmark/WorldBookmarkStyle.cpp
new file mode 100644
--- /dev/null
+++ b/Engine/Source/Editor/WorldBookmark/Private/WorldBookmark/WorldBookmarkStyle.cpp
@@ -0,0 +1,51 @@
+// Copyright Epic Games, Inc. All Rights Reserved.
+
+#include "WorldBookmark/WorldBookmarkStyle.h"
+
+#include "Styling/SlateStyleMacros.h"
+
+#include "Styling/AppStyle.h"
+#include "Styling/StyleColors.h"
+
+FWorldBookmarkStyle::FWorldBookmarkStyle()
+	: FSlateStyleSet("WorldBookmark")
+{
+	const FVector2D Icon16x16(16.0f, 16.0f);
+	const FVector2D Icon20x20(20.0f, 20.0f);
+	const FVector2D Icon64x64(64.0f, 64.0f);
+
+	SetParentStyleName(FAppStyle::GetAppStyleSetName());
+
+	SetContentRoot(FPaths::EngineContentDir() / TEXT("Editor/Slate"));
+	SetCoreContentRoot(FPaths::EngineContentDir() / TEXT("Slate"));
+
+	// Asset
+	Set("ClassIcon.WorldBookmark", new IMAGE_BRUSH_SVG("Starship/AssetIcons/WorldBookmark_16", Icon16x16));
+	Set("ClassThumbnail.WorldBookmark", new IMAGE_BRUSH_SVG("Starship/AssetIcons/WorldBookmark_64", Icon64x64));
+
+	// Commands
+	Set("WorldBookmark.LoadBookmark", new IMAGE_BRUSH_SVG("Starship/Common/NextArrow", Icon16x16));
+	Set("WorldBookmark.UpdateBookmark", new CORE_IMAGE_BRUSH_SVG("Starship/Common/Update", Icon20x20));
+	Set("WorldBookmark.CreateBookmark", new CORE_IMAGE_BRUSH_SVG("Starship/Common/plus", Icon16x16, FStyleColors::AccentGreen));
+	Set("WorldBookmark.AddToFavorite", new IMAGE_BRUSH("Icons/Star_16x", Icon16x16));
+	Set("WorldBookmark.RemoveFromFavorite", new IMAGE_BRUSH("Icons/EmptyStar_16x", Icon16x16));
+
+	// Icons
+	Set("WorldBookmark.TabIcon", new IMAGE_BRUSH_SVG("Starship/AssetIcons/WorldBookmark_16", Icon16x16));
+	Set("WorldBookmark.IsFavorite", new IMAGE_BRUSH("Icons/Star_16x", Icon16x16));
+	Set("WorldBookmark.IsNotFavorite", new IMAGE_BRUSH("Icons/EmptyStar_16x", Icon16x16));
+	Set("WorldBookmark.RecentlyUsed", new IMAGE_BRUSH("Icons/Clock_24x", Icon16x16));
+
+	FSlateStyleRegistry::RegisterSlateStyle(*this);
+}
+
+FWorldBookmarkStyle::~FWorldBookmarkStyle()
+{
+	FSlateStyleRegistry::UnRegisterSlateStyle(*this);
+}
+
+FWorldBookmarkStyle& FWorldBookmarkStyle::Get()
+{
+	static FWorldBookmarkStyle Inst;
+	return Inst;
+}
diff --git a/Engine/Source/Editor/WorldBookmark/Private/WorldBookmark/WorldBookmarkStyle.h b/Engine/Source/Editor/WorldBookmark/Private/WorldBookmark/WorldBookmarkStyle.h
new file mode 100644
--- /dev/null
+++ b/Engine/Source/Editor/WorldBookmark/Private/WorldBookmark/WorldBookmarkStyle.h
@@ -0,0 +1,15 @@
+// Copyright Epic Games, Inc. All Rights Reserved.
+
+#pragma once
+
+#include "Styling/SlateStyleRegistry.h"
+#include "Styling/SlateStyle.h"
+
+class FWorldBookmarkStyle final : public FSlateStyleSet
+{
+public:
+	FWorldBookmarkStyle();
+	virtual ~FWorldBookmarkStyle();
+
+	static FWorldBookmarkStyle& Get();
+};
