World Bookmark - When deleting a world, ask if related World Bookmarks should also be deleted

diff --git a/Engine/Source/Editor/WorldBookmark/Private/WorldBookmark/WorldBookmark.cpp b/Engine/Source/Editor/WorldBookmark/Private/WorldBookmark/WorldBookmark.cpp
--- a/Engine/Source/Editor/WorldBookmark/Private/WorldBookmark/WorldBookmark.cpp
+++ b/Engine/Source/Editor/WorldBookmark/Private/WorldBookmark/WorldBookmark.cpp
@@ -68,6 +68,12 @@ void UWorldBookmark::OverridePerObjectConfigSection(FString& SectionName)
 	SectionName = FString::Printf(TEXT("WorldBookmark %s"), *BookmarkGuid.ToString());
 }
 
+FName UWorldBookmark::GetWorldNameAssetTag()
+{
+	static const FName NAME_WorldName = "WorldName";
+	return NAME_WorldName;
+}
+
 void UWorldBookmark::GetAssetRegistryTags(FAssetRegistryTagsContext Context) const
 {
 	Super::GetAssetRegistryTags(Context);
@@ -80,8 +86,7 @@ void UWorldBookmark::GetAssetRegistryTags(FAssetRegistryTagsContext Context) con
 			FSoftObjectPath WorldPath = FSoftObjectPath(World.ToString());
 			UAssetRegistryHelpers::FixupRedirectedAssetPath(WorldPath);
 
-			static const FName NAME_BookmarkWorld = "WorldName";
-			Context.AddTag(FAssetRegistryTag(NAME_BookmarkWorld, WorldPath.ToString(), FAssetRegistryTag::TT_Hidden));
+			Context.AddTag(FAssetRegistryTag(GetWorldNameAssetTag(), WorldPath.ToString(), FAssetRegistryTag::TT_Hidden));
 		}
 	}
 }
diff --git a/Engine/Source/Editor/WorldBookmark/Private/WorldBookmark/WorldBookmarkModule.cpp b/Engine/Source/Editor/WorldBookmark/Private/WorldBookmark/WorldBookmarkModule.cpp
--- a/Engine/Source/Editor/WorldBookmark/Private/WorldBookmark/WorldBookmarkModule.cpp
+++ b/Engine/Source/Editor/WorldBookmark/Private/WorldBookmark/WorldBookmarkModule.cpp
@@ -7,8 +7,12 @@
 #include "WorldBookmark/WorldBookmarkStyle.h"
 #include "WorldBookmark/Browser/SWorldBookmarkBrowser.h"
 
+#include "AssetRegistry/IAssetRegistry.h"
+#include "AssetRegistry/AssetData.h"
 #include "LevelEditor.h"
+#include "Misc/MessageDialog.h"
 #include "Modules/ModuleManager.h"
+#include "ObjectTools.h"
 #include "PropertyEditorModule.h"
 #include "Widgets/Docking/SDockTab.h"
 #include "WorkspaceMenuStructure.h"
@@ -33,6 +37,8 @@ void FWorldBookmarkModule::StartupModule()
 
 		ClassesToUnregisterOnShutdown.Add(UWorldBookmark::StaticClass()->GetFName());
 		PropertyModule.RegisterCustomClassLayout(UWorldBookmark::StaticClass()->GetFName(), FOnGetDetailCustomizationInstance::CreateStatic(&FWorldBookmarkDetailsCustomization::MakeInstance));
+
+		OnAddExtraObjectsToDeleteDelegateHandle = FEditorDelegates::OnAddExtraObjectsToDelete.AddRaw(this, &FWorldBookmarkModule::OnAddExtraObjectsToDelete);
 	}
 }
 
@@ -59,6 +65,8 @@ void FWorldBookmarkModule::ShutdownModule()
 		}
 
 		ClassesToUnregisterOnShutdown.Empty();
+
+		FEditorDelegates::OnAddExtraObjectsToDelete.Remove(OnAddExtraObjectsToDeleteDelegateHandle);
 	}
 }
 
@@ -93,4 +101,68 @@ TSharedRef<SWidget> FWorldBookmarkModule::CreateWorldBookmarkBrowser()
 	return SNew(SWorldBookmarkBrowser);
 }
 
+void FWorldBookmarkModule::OnAddExtraObjectsToDelete(const TArray<UObject*>& InObjectsToDelete, TSet<UObject*>& OutSecondaryObjects)
+{
+	// Gather a list of worlds
+	TArray<const UWorld*> WorldsPendingDelete;
+	for (const UObject* Object : InObjectsToDelete)
+	{
+		if (const UWorld* World = Cast<UWorld>(Object))
+		{
+			WorldsPendingDelete.Add(World);
+		}
+	}
+
+	// If we are deleting worlds
+	if (!WorldsPendingDelete.IsEmpty())
+	{
+		// Scan the asset registry and look for bookmarks for these worlds
+		FARFilter ArFilter;
+		ArFilter.ClassPaths.Add(UWorldBookmark::StaticClass()->GetClassPathName());
+		ArFilter.bRecursiveClasses = true;
+		for (const UWorld* WorldPendingDelete : WorldsPendingDelete)
+		{
+			ArFilter.TagsAndValues.Emplace(UWorldBookmark::GetWorldNameAssetTag(), WorldPendingDelete->GetPathName());
+		}
+		TArray<FAssetData> AssetsData;
+		IAssetRegistry::GetChecked().GetAssets(ArFilter, AssetsData);
+
+		// We've found some bookmarks related to these worlds, ask the user if they want to delete them
+		if (!AssetsData.IsEmpty())
+		{
+			FTextBuilder TextBuilder;
+			
+			int32 NbLines = 0;
+			const int32 MaxNbLines = 5;
+			for (const FAssetData& AssetData : AssetsData)
+			{
+				if (++NbLines > MaxNbLines)
+				{
+					TextBuilder.AppendLineFormat(LOCTEXT("DeleteMatchingBookmarks_OverflowList", "    ({0} more assets...)"), FText::AsNumber(AssetsData.Num() - MaxNbLines));
+					break;
+				}
+
+				TextBuilder.AppendLineFormat(LOCTEXT("DeleteMatchingBookmarks_AssetList", "    {0}"), FText::FromString(AssetData.PackageName.ToString()));
+			}
+
+			const FText AllAssetsText = TextBuilder.ToText();
+			const FText MessageBoxTitle = LOCTEXT("DeleteMatchingBookmarks_Title", "Delete World Bookmark(s)?");
+			const FText MessageBoxTextSingleAsset = FText::Format(LOCTEXT("DeleteMatchingBookmark_Text", "This world is referenced by a World Bookmark. Do you wish to delete this asset too?\n{0}"), AllAssetsText);
+			const FText MessageBoxTextMultipleAssets = FText::Format(LOCTEXT("DeleteMatchingBookmarks_Text", "This world is referenced by {0} World Bookmarks. Do you wish to delete those assets too?\n{1}"), FText::AsNumber(AssetsData.Num()), AllAssetsText);
+
+			const EAppReturnType::Type Response = FMessageDialog::Open(EAppMsgCategory::Info, EAppMsgType::YesNo, AssetsData.Num() > 1 ? MessageBoxTextMultipleAssets : MessageBoxTextSingleAsset, MessageBoxTitle);
+			if (Response == EAppReturnType::Yes)
+			{
+				for (const FAssetData& AssetData : AssetsData)
+				{
+					if (UWorldBookmark* WorldBookmark = Cast<UWorldBookmark>(AssetData.GetAsset()))
+					{
+						OutSecondaryObjects.Add(WorldBookmark);
+					}
+				}
+			}
+		}
+	}
+}
+
 #undef LOCTEXT_NAMESPACE
diff --git a/Engine/Source/Editor/WorldBookmark/Private/WorldBookmark/WorldBookmarkModule.h b/Engine/Source/Editor/WorldBookmark/Private/WorldBookmark/WorldBookmarkModule.h
--- a/Engine/Source/Editor/WorldBookmark/Private/WorldBookmark/WorldBookmarkModule.h
+++ b/Engine/Source/Editor/WorldBookmark/Private/WorldBookmark/WorldBookmarkModule.h
@@ -2,6 +2,7 @@
 
 #pragma once
 
+#include "Delegates/IDelegateInstance.h"
 #include "Modules/ModuleInterface.h"
 #include "Templates/SharedPointer.h"
 
@@ -31,5 +32,9 @@ private:
 	TSharedRef<SDockTab> SpawnWorldBookmarkBrowserTab(const FSpawnTabArgs& Args);
 	TSharedRef<SWidget> CreateWorldBookmarkBrowser();
 
+	void OnAddExtraObjectsToDelete(const TArray<UObject*>& InObjectsToDelete, TSet<UObject*>& OutSecondaryObjects);
+
+	FDelegateHandle OnAddExtraObjectsToDeleteDelegateHandle;
+
 	TArray<FName> ClassesToUnregisterOnShutdown;
 };
diff --git a/Engine/Source/Editor/WorldBookmark/Public/WorldBookmark/WorldBookmark.h b/Engine/Source/Editor/WorldBookmark/Public/WorldBookmark/WorldBookmark.h
--- a/Engine/Source/Editor/WorldBookmark/Public/WorldBookmark/WorldBookmark.h
+++ b/Engine/Source/Editor/WorldBookmark/Public/WorldBookmark/WorldBookmark.h
@@ -39,6 +39,9 @@ class UWorldBookmark : public UObject
 	/** Store the last time that bookmark was loaded, in UTC. */
 	void SetUserLastLoadedTimeStampUTC(const FDateTime& InLastLoadedTimeStampUTC);
 
+    /** Get the asset registry tag used to store the matching world's name. */
+	static FName GetWorldNameAssetTag();
+
 private:
 	friend class FWorldBookmarkDetailsCustomization;
 
