World Bookmark - Added "Play From Here" and "Move Camera Here" actions
* When performing these actions on a bookmark for a different world that the current world, ask user confirmation
* Also prevented most bookmark actions from being executed while in PIE

diff --git a/Engine/Source/Editor/LevelEditor/Private/LevelEditorCameraEditorState.cpp b/Engine/Source/Editor/LevelEditor/Private/LevelEditorCameraEditorState.cpp
--- a/Engine/Source/Editor/LevelEditor/Private/LevelEditorCameraEditorState.cpp
+++ b/Engine/Source/Editor/LevelEditor/Private/LevelEditorCameraEditorState.cpp
@@ -19,6 +19,21 @@ FText ULevelEditorCameraEditorState::GetCategoryText() const
 	return FText(LOCTEXT("LevelEditorCameraEditorStateCategoryText", "Camera"));
 }
 
+const FVector& ULevelEditorCameraEditorState::GetCameraLocation() const
+{
+	return CameraLocation;
+}
+
+const FRotator& ULevelEditorCameraEditorState::GetCameraRotation() const
+{
+	return CameraRotation;
+}
+
+const float ULevelEditorCameraEditorState::GetCameraFOVAngle() const
+{
+	return CameraFOVAngle;
+}
+
 UEditorState::FOperationResult ULevelEditorCameraEditorState::CaptureState()
 {
 	TSharedPtr<SLevelViewport> LevelViewport = InternalEditorLevelLibrary::GetActiveLevelViewport();
diff --git a/Engine/Source/Editor/LevelEditor/Public/LevelEditorCameraEditorState.h b/Engine/Source/Editor/LevelEditor/Public/LevelEditorCameraEditorState.h
--- a/Engine/Source/Editor/LevelEditor/Public/LevelEditorCameraEditorState.h
+++ b/Engine/Source/Editor/LevelEditor/Public/LevelEditorCameraEditorState.h
@@ -2,16 +2,24 @@
 
 #pragma once
 
-#include "EditorState/EditorState.h"
+#include "EditorState/WorldEditorState.h"
 #include "LevelEditorCameraEditorState.generated.h"
 
-UCLASS()
-class ULevelEditorCameraEditorState : public UEditorState
+class SLevelViewport;
+
+UCLASS(MinimalAPI)
+class ULevelEditorCameraEditorState : public UWorldDependantEditorState
 {
 	GENERATED_UCLASS_BODY()
 
+public:
 	//~ Begin UEditorState Interface
 	virtual FText GetCategoryText() const override;
+
+	LEVELEDITOR_API const FVector& GetCameraLocation() const;
+	LEVELEDITOR_API const FRotator& GetCameraRotation() const;
+	LEVELEDITOR_API const float GetCameraFOVAngle() const;
+
 private:
 	virtual FOperationResult CaptureState() override;
 	virtual FOperationResult RestoreState() const override;
diff --git a/Engine/Source/Editor/UnrealEd/Private/EditorState/EditorStateSubsystem.cpp b/Engine/Source/Editor/UnrealEd/Private/EditorState/EditorStateSubsystem.cpp
--- a/Engine/Source/Editor/UnrealEd/Private/EditorState/EditorStateSubsystem.cpp
+++ b/Engine/Source/Editor/UnrealEd/Private/EditorState/EditorStateSubsystem.cpp
@@ -75,16 +75,31 @@ void UEditorStateSubsystem::OutputOperationResult(const UEditorState* InState, b
 		break;
 	}
 }
-	
+
 void UEditorStateSubsystem::CaptureEditorState(FEditorStateCollection& OutState, UObject* InStateOuter)
+{
+	CaptureEditorState(OutState, {}, InStateOuter);
+}
+
+void UEditorStateSubsystem::CaptureEditorState(FEditorStateCollection& OutState, const TArray<TSubclassOf<UEditorState>>& InEditorStateTypeFilter, UObject* InStateOuter)
 {
 	UE_LOG(LogEditorState, Log, TEXT("Capturing editor state..."));
 
-	OutState.EditorStates.Reset();
+	if (InEditorStateTypeFilter.IsEmpty())
+	{
+		OutState.EditorStates.Reset();
 
-	for (TSubclassOf<UEditorState> EditorStateType : RegisteredEditorStateTypes)
+		for (TSubclassOf<UEditorState> EditorStateType : RegisteredEditorStateTypes)
+		{
+			OutState.EditorStates.Emplace(EditorStateType, NewObject<UEditorState>(InStateOuter, EditorStateType));
+		}
+	}
+	else
 	{
-		OutState.EditorStates.Emplace(EditorStateType, NewObject<UEditorState>(InStateOuter, EditorStateType));
+		for (TSubclassOf<UEditorState> EditorStateType : InEditorStateTypeFilter)
+		{
+			OutState.EditorStates.Emplace(EditorStateType, NewObject<UEditorState>(InStateOuter, EditorStateType));
+		}
 	}
 
 	OutState.ForEachState([&OutState, this](UEditorState* StateToCapture, bool bCapturedDependencies)
@@ -111,7 +126,7 @@ void UEditorStateSubsystem::CaptureEditorState(FEditorStateCollection& OutState,
 		}
 
 		return bCaptureSuccess;
-	});
+	}, InEditorStateTypeFilter);
 
 	// Remove invalidated states
 	for (FEditorStateCollection::FEditorStateMap::TIterator It(OutState.EditorStates); It; ++It)
@@ -126,6 +141,11 @@ void UEditorStateSubsystem::CaptureEditorState(FEditorStateCollection& OutState,
 }
 
 void UEditorStateSubsystem::RestoreEditorState(const FEditorStateCollection& InState)
+{
+	RestoreEditorState(InState, {});
+}
+
+void UEditorStateSubsystem::RestoreEditorState(const FEditorStateCollection& InState, const TArray<TSubclassOf<UEditorState>>& InEditorStateTypeFilter)
 {
 	UE_LOG(LogEditorState, Log, TEXT("Restoring editor state..."));
 
@@ -142,7 +162,7 @@ void UEditorStateSubsystem::RestoreEditorState(const FEditorStateCollection& InS
 
 		return Result == UEditorState::FOperationResult::Success ||
 			   Result == UEditorState::FOperationResult::Warning;
-	});
+	}, InEditorStateTypeFilter);
 
 	UE_LOG(LogEditorState, Log, TEXT("Restored editor state... DONE"));
 }
diff --git a/Engine/Source/Editor/UnrealEd/Private/EditorState/WorldEditorState.cpp b/Engine/Source/Editor/UnrealEd/Private/EditorState/WorldEditorState.cpp
--- a/Engine/Source/Editor/UnrealEd/Private/EditorState/WorldEditorState.cpp
+++ b/Engine/Source/Editor/UnrealEd/Private/EditorState/WorldEditorState.cpp
@@ -43,7 +43,7 @@ UEditorState::FOperationResult UWorldEditorState::RestoreState() const
 {
 	if (World.IsNull())
 	{
-		return FOperationResult(FOperationResult::Failure, LOCTEXT("RestoreStateFailure_UnsavedWorld", "Bookmark world is invalid"));
+		return FOperationResult(FOperationResult::Failure, LOCTEXT("RestoreStateFailure_UnsavedWorld", "World is invalid"));
 	}
 
 	FString WorldPackageString = World.GetLongPackageName();
@@ -52,7 +52,7 @@ UEditorState::FOperationResult UWorldEditorState::RestoreState() const
 
 	if (IAssetRegistry::GetChecked().TryGetAssetByObjectPath(World.ToSoftObjectPath(), WorldAsset) != UE::AssetRegistry::EExists::Exists)
 	{
-		return FOperationResult(FOperationResult::Failure, FText::Format(LOCTEXT("RestoreStateFailure_TargetWorldNotFound", "Bookmark world {0} couldn't be resolved"), WorldPackageText));
+		return FOperationResult(FOperationResult::Failure, FText::Format(LOCTEXT("RestoreStateFailure_TargetWorldNotFound", "World {0} couldn't be resolved"), WorldPackageText));
 	}
 
 	if (WorldAsset.AssetClassPath != UWorld::StaticClass()->GetClassPathName())
@@ -64,7 +64,7 @@ UEditorState::FOperationResult UWorldEditorState::RestoreState() const
 	UWorld* CurrentWorld = GEditor->GetEditorWorldContext().World();
 	if (CurrentWorld == World)
 	{
-		return FOperationResult(FOperationResult::Success, FText::Format(LOCTEXT("RestoreStateSuccess_WorldAlreadyLoaded", "Bookmark world {0} was already active"), WorldPackageText));
+		return FOperationResult(FOperationResult::Success, FText::Format(LOCTEXT("RestoreStateSuccess_WorldAlreadyLoaded", "World {0} was already active"), WorldPackageText));
 	}
 
 	// If there are any unsaved changes to the current level, see if the user wants to save those first
diff --git a/Engine/Source/Editor/UnrealEd/Public/EditorState/EditorStateCollection.h b/Engine/Source/Editor/UnrealEd/Public/EditorState/EditorStateCollection.h
--- a/Engine/Source/Editor/UnrealEd/Public/EditorState/EditorStateCollection.h
+++ b/Engine/Source/Editor/UnrealEd/Public/EditorState/EditorStateCollection.h
@@ -35,17 +35,17 @@ struct FEditorStateCollection
 	}
 
 	/** Iterate over each state, making sure iteration is done over dependant states first */
-	void ForEachState(TFunctionRef<bool(const UEditorState*, bool)> InFunc) const
+	void ForEachState(TFunctionRef<bool(const UEditorState*, bool)> InFunc, const TArray<TSubclassOf<UEditorState>>& InEditorStateTypeFilter = {}) const
 	{
-		ForEachState(this, InFunc);
+		ForEachState(this, InFunc, InEditorStateTypeFilter);
 	}
 
 	/** Iterate over each state, making sure iteration is done over dependant states first */
-	void ForEachState(TFunctionRef<bool(UEditorState*, bool)> InFunc)
+	void ForEachState(TFunctionRef<bool(UEditorState*, bool)> InFunc, const TArray<TSubclassOf<UEditorState>>& InEditorStateTypeFilter = {})
 	{
-		ForEachState(this, InFunc);
+		ForEachState(this, InFunc, InEditorStateTypeFilter);
 	}
-	
+
 	TArray<UEditorState*> GetStates() const
 	{
 		TArray<UEditorState*> Values;
@@ -60,7 +60,7 @@ struct FEditorStateCollection
 
 private:
 	template <typename ThisType, typename FuncType>
-	static void ForEachState(ThisType* Self, FuncType InFunc)
+	static void ForEachState(ThisType* Self, FuncType InFunc, const TArray<TSubclassOf<UEditorState>>& InEditorStateTypeFilter = {})
 	{
 		// Keep track of processed states to avoid cycles
 		TMap<TSubclassOf<UEditorState>, bool> ProcessedStates;
@@ -103,6 +103,11 @@ private:
 
 		for (const auto& KeyValue : Self->EditorStates)
 		{
+			if (!InEditorStateTypeFilter.IsEmpty() && !InEditorStateTypeFilter.Contains(KeyValue.Value->GetClass()))
+			{
+				continue;
+			}
+
 			ProcessState(KeyValue.Value);
 		}
 	}
diff --git a/Engine/Source/Editor/UnrealEd/Public/EditorState/EditorStateSubsystem.h b/Engine/Source/Editor/UnrealEd/Public/EditorState/EditorStateSubsystem.h
--- a/Engine/Source/Editor/UnrealEd/Public/EditorState/EditorStateSubsystem.h
+++ b/Engine/Source/Editor/UnrealEd/Public/EditorState/EditorStateSubsystem.h
@@ -44,11 +44,24 @@ public:
 	 */
 	UNREALED_API void CaptureEditorState(FEditorStateCollection& OutState, UObject* InStateOuter);
 
+	/**
+	 * Captures the state of the editor using all the registered EditorState subclasses.
+	 * @param OutState		Collection of states that will be populated with state data.
+	 * @param InStateOuter	Outer to use for the UEditorState objects that may be created by this method.
+	 */
+	UNREALED_API void CaptureEditorState(FEditorStateCollection& OutState, const TArray<TSubclassOf<UEditorState>>& InEditorStateTypeFilter, UObject* InStateOuter);
+
 	/**
 	 * Restore the state of the editor using a state collection previously captured with CaptureEditorState().
 	 * @param InState		Collection of states to restore.
 	 */
 	UNREALED_API void RestoreEditorState(const FEditorStateCollection& InState);
+
+	/**
+	 * Restore the state of the editor using a state collection previously captured with CaptureEditorState().
+	 * @param InState		Collection of states to restore.
+	 */
+	UNREALED_API void RestoreEditorState(const FEditorStateCollection& InState, const TArray<TSubclassOf<UEditorState>>& InEditorStateTypeFilter);
 	
 private:
 	UNREALED_API void RegisterEditorStateType(TSubclassOf<UEditorState> InEditorStateType);
diff --git a/Engine/Source/Editor/WorldBookmark/Private/WorldBookmark/Browser/SWorldBookmarkBrowser.cpp b/Engine/Source/Editor/WorldBookmark/Private/WorldBookmark/Browser/SWorldBookmarkBrowser.cpp
--- a/Engine/Source/Editor/WorldBookmark/Private/WorldBookmark/Browser/SWorldBookmarkBrowser.cpp
+++ b/Engine/Source/Editor/WorldBookmark/Private/WorldBookmark/Browser/SWorldBookmarkBrowser.cpp
@@ -12,10 +12,14 @@
 #include "Editor.h"
 #include "Framework/Application/SlateApplication.h"
 #include "Framework/Commands/GenericCommands.h"
+#include "Kismet2/DebuggerCommands.h"
 #include "LevelEditor.h"
+#include "LevelEditorCameraEditorState.h"
+#include "Misc/MessageDialog.h"
 #include "Modules/ModuleManager.h"
 #include "ObjectTools.h"
 #include "PropertyEditorModule.h"
+#include "SLevelViewport.h"
 #include "SourceControlHelpers.h"
 #include "SSimpleButton.h"
 #include "Toolkits/GlobalEditorCommonCommands.h"
@@ -217,7 +221,7 @@ void SWorldBookmarkBrowser::Construct(const FArguments& InArgs)
 					[
 						SNew(SImage)
 						.ColorAndOpacity(FSlateColor::UseForeground())
-						.Image(FWorldBookmarkStyle::Get().GetBrush("WorldBookmark.Create"))
+						.Image(FWorldBookmarkStyle::Get().GetBrush("WorldBookmark.CreateBookmark"))
 					]
 				]
 
@@ -273,11 +277,11 @@ void SWorldBookmarkBrowser::Construct(const FArguments& InArgs)
 
 	Commands->MapAction(FWorldBookmarkCommands::Get().LoadBookmark,
 		FExecuteAction::CreateSP(this, &SWorldBookmarkBrowser::LoadSelectedBookmark),
-		FCanExecuteAction::CreateSP(this, &SWorldBookmarkBrowser::IsValidBookmarkSelected));
+		FCanExecuteAction::CreateSP(this, &SWorldBookmarkBrowser::CanExecuteBookmarkAction));
 
 	Commands->MapAction(FWorldBookmarkCommands::Get().UpdateBookmark,
 		FExecuteAction::CreateSP(this, &SWorldBookmarkBrowser::UpdateSelectedBookmark),
-		FCanExecuteAction::CreateSP(this, &SWorldBookmarkBrowser::IsValidBookmarkSelected));
+		FCanExecuteAction::CreateSP(this, &SWorldBookmarkBrowser::CanExecuteBookmarkAction));
 	
 	Commands->MapAction(FWorldBookmarkCommands::Get().AddToFavorite,
 		FExecuteAction::CreateSP(this, &SWorldBookmarkBrowser::AddSelectedBookmarkToFavorites),
@@ -290,6 +294,14 @@ void SWorldBookmarkBrowser::Construct(const FArguments& InArgs)
 		FCanExecuteAction::CreateSP(this, &SWorldBookmarkBrowser::IsValidBookmarkSelected),
 		FIsActionChecked(),
 		FIsActionButtonVisible::CreateSP(this, &SWorldBookmarkBrowser::IsSelectedBookmarkFavorite));
+
+	Commands->MapAction(FWorldBookmarkCommands::Get().PlayFromLocation,
+		FExecuteAction::CreateSP(this, &SWorldBookmarkBrowser::PlayFromSelectedBookmarkLocation),
+		FCanExecuteAction::CreateSP(this, &SWorldBookmarkBrowser::CanGoToSelectedBookmarkLocation));
+
+	Commands->MapAction(FWorldBookmarkCommands::Get().MoveCameraToLocation,
+		FExecuteAction::CreateSP(this, &SWorldBookmarkBrowser::GoToSelectedBookmarkLocation),
+		FCanExecuteAction::CreateSP(this, &SWorldBookmarkBrowser::CanGoToSelectedBookmarkLocation));
 }
 
 TSharedRef<SWidget> SWorldBookmarkBrowser::GetSettingsMenuContent()
@@ -891,6 +903,8 @@ TSharedPtr<SWidget> SWorldBookmarkBrowser::OnContextMenuOpening()
 		MenuBuilder.AddMenuEntry(FWorldBookmarkCommands::Get().UpdateBookmark);
 		MenuBuilder.AddMenuEntry(FWorldBookmarkCommands::Get().AddToFavorite);
 		MenuBuilder.AddMenuEntry(FWorldBookmarkCommands::Get().RemoveFromFavorite);
+		MenuBuilder.AddMenuEntry(FWorldBookmarkCommands::Get().PlayFromLocation);
+		MenuBuilder.AddMenuEntry(FWorldBookmarkCommands::Get().MoveCameraToLocation);
 	}
 
 	MenuBuilder.EndSection();
@@ -1019,6 +1033,18 @@ void SWorldBookmarkBrowser::OnUncontrolledChangelistModuleChanged()
 	RefreshItems();
 }
 
+bool SWorldBookmarkBrowser::CanExecuteBookmarkAction() const
+{
+	// Disable actions in PIE
+	if (GEditor->PlayWorld != nullptr || GEditor->bIsSimulatingInEditor)
+	{
+		return false;
+	}
+
+	// Disable actions if no bookmark is selected
+	return IsValidBookmarkSelected();
+}
+
 void SWorldBookmarkBrowser::LoadSelectedBookmark()
 {
 	if (UWorldBookmark* SelectedBookmark = GetSelectedBookmark())
@@ -1181,6 +1207,82 @@ void SWorldBookmarkBrowser::DeleteSelectedBookmark()
 	}
 }
 
+static bool AskForWorldChangeConfirmation(UWorldBookmark* WorldBookmark)
+{
+	TSoftObjectPtr<UWorld> BookmarkWorld = WorldBookmark->GetEditorState<UWorldEditorState>()->GetStateWorld();
+	if (BookmarkWorld != GetCurrentWorld())
+	{
+		const FText MessageBoxTitle = LOCTEXT("WorldBookmark_DifferentWorldDlg_Title", "Open Bookmark World?");
+		const FText MessageBoxText = LOCTEXT("WorldBookmark_DifferentWorldDlg_Text", "This bookmark is for a different world, are you sure you want to open it?");
+		const EAppReturnType::Type Response = FMessageDialog::Open(EAppMsgType::YesNo, MessageBoxText, MessageBoxTitle);
+		if (Response != EAppReturnType::Yes)
+		{
+			return false;
+		}
+	}
+
+	return true;
+}
+
+bool SWorldBookmarkBrowser::CanGoToSelectedBookmarkLocation() const
+{
+	if (!CanExecuteBookmarkAction())
+	{
+		return false;
+	}
+
+	if (UWorldBookmark* SelectedBookmark = GetSelectedBookmark())
+	{
+		if (SelectedBookmark->HasEditorState<ULevelEditorCameraEditorState>())
+		{
+			return true;
+		}
+	}
+
+	return false;
+}
+
+void SWorldBookmarkBrowser::PlayFromSelectedBookmarkLocation()
+{
+	if (UWorldBookmark* SelectedBookmark = GetSelectedBookmark())
+	{
+		// Test if state contains LevelEditorCameraState ?
+		if (SelectedBookmark->HasEditorState<ULevelEditorCameraEditorState>())
+		{
+			if (!AskForWorldChangeConfirmation(SelectedBookmark))
+			{
+				return;
+			}
+
+			SelectedBookmark->LoadStates({ UWorldEditorState::StaticClass() });
+
+			if (const ULevelEditorCameraEditorState* CameraEditorState = SelectedBookmark->GetEditorState<ULevelEditorCameraEditorState>())
+			{
+				FLevelEditorModule& LevelEditorModule = FModuleManager::GetModuleChecked<FLevelEditorModule>("LevelEditor");
+				TSharedPtr<IAssetViewport> ActiveViewport = LevelEditorModule.GetFirstActiveViewport();
+				FPlayWorldCommandCallbacks::StartPlayFromHere(CameraEditorState->GetCameraLocation(), CameraEditorState->GetCameraRotation(), ActiveViewport);
+			}
+		}
+	}
+}
+
+void SWorldBookmarkBrowser::GoToSelectedBookmarkLocation()
+{
+	if (UWorldBookmark* SelectedBookmark = GetSelectedBookmark())
+	{
+		// Test if state contains LevelEditorCameraState ?
+		if (SelectedBookmark->HasEditorState<ULevelEditorCameraEditorState>())
+		{
+			if (!AskForWorldChangeConfirmation(SelectedBookmark))
+			{
+				return;
+			}
+
+			SelectedBookmark->LoadStates({ ULevelEditorCameraEditorState::StaticClass() });
+		}
+	}
+}
+
 void SWorldBookmarkBrowser::Tick(const FGeometry& AllottedGeometry, const double InCurrentTime, const float InDeltaTime)
 {
 	TGuardValue<bool> IsInTickGuard(bIsInTick, true);
diff --git a/Engine/Source/Editor/WorldBookmark/Private/WorldBookmark/Browser/SWorldBookmarkBrowser.h b/Engine/Source/Editor/WorldBookmark/Private/WorldBookmark/Browser/SWorldBookmarkBrowser.h
--- a/Engine/Source/Editor/WorldBookmark/Private/WorldBookmark/Browser/SWorldBookmarkBrowser.h
+++ b/Engine/Source/Editor/WorldBookmark/Private/WorldBookmark/Browser/SWorldBookmarkBrowser.h
@@ -100,6 +100,8 @@ private:
 	UWorldBookmark* GetSelectedBookmark() const;
 	void SetSelectedBookmark(UWorldBookmark* InBookmark);
 
+	bool CanExecuteBookmarkAction() const;
+
 	void LoadSelectedBookmark();
 	void UpdateSelectedBookmark();
 	void FindSelectedBookmarkInContentBrowser();
@@ -113,6 +115,10 @@ private:
 	void RequestRenameSelectedBookmark();
 	void DeleteSelectedBookmark();
 
+	bool CanGoToSelectedBookmarkLocation() const;
+	void PlayFromSelectedBookmarkLocation();
+	void GoToSelectedBookmarkLocation();
+
 private:
 	FDelegateHandle OnAssetAddedHandle;
 	FDelegateHandle OnAssetRemovedHandle;
diff --git a/Engine/Source/Editor/WorldBookmark/Private/WorldBookmark/WorldBookmark.cpp b/Engine/Source/Editor/WorldBookmark/Private/WorldBookmark/WorldBookmark.cpp
--- a/Engine/Source/Editor/WorldBookmark/Private/WorldBookmark/WorldBookmark.cpp
+++ b/Engine/Source/Editor/WorldBookmark/Private/WorldBookmark/WorldBookmark.cpp
@@ -117,11 +117,16 @@ bool UWorldBookmark::CanLoad(FText* FailureReason) const
 }
 
 void UWorldBookmark::Load()
+{
+	LoadStates({});
+}
+
+void UWorldBookmark::LoadStates(const TArray<TSubclassOf<UEditorState>>& InStatesToLoad)
 {
 	FText CanLoadFailureReason;
 	if (CanLoad(&CanLoadFailureReason))
 	{
-		UEditorStateSubsystem::Get()->RestoreEditorState(EditorState);
+		UEditorStateSubsystem::Get()->RestoreEditorState(EditorState, InStatesToLoad);
 
 		SetUserLastLoadedTimeStampUTC(FDateTime::UtcNow());
 		SaveConfig();
@@ -138,7 +143,7 @@ void UWorldBookmark::Load()
 			Info.Image = FCoreStyle::Get().GetBrush(TEXT("MessageLog.Warning"));
 			FSlateNotificationManager::Get().AddNotification(Info);
 		}
-		
+
 		UE_LOG(LogWorldBookmark, Error, TEXT("Failed to load bookmark: %s"), *CanLoadFailureReason.ToString());
 	}
 }
@@ -190,6 +195,11 @@ bool UWorldBookmark::CanUpdate(FText* FailureReason) const
 }
 
 void UWorldBookmark::Update()
+{
+	UpdateStates({});
+}
+
+void UWorldBookmark::UpdateStates(const TArray<TSubclassOf<UEditorState>>& InStatesToUpdate)
 {
 	FText CanUpdateFailureReason;
 	if (CanUpdate(&CanUpdateFailureReason))
@@ -198,7 +208,7 @@ void UWorldBookmark::Update()
 
 		Modify();
 
-		UEditorStateSubsystem::Get()->CaptureEditorState(EditorState, this);
+		UEditorStateSubsystem::Get()->CaptureEditorState(EditorState, InStatesToUpdate, this);
 
 		PostEditChange();
 	}
diff --git a/Engine/Source/Editor/WorldBookmark/Private/WorldBookmark/WorldBookmarkCommands.cpp b/Engine/Source/Editor/WorldBookmark/Private/WorldBookmark/WorldBookmarkCommands.cpp
--- a/Engine/Source/Editor/WorldBookmark/Private/WorldBookmark/WorldBookmarkCommands.cpp
+++ b/Engine/Source/Editor/WorldBookmark/Private/WorldBookmark/WorldBookmarkCommands.cpp
@@ -20,6 +20,8 @@ void FWorldBookmarkCommands::RegisterCommands()
 	UI_COMMAND(UpdateBookmark, "Update Bookmark", "Update this bookmark using the current state of the world.", EUserInterfaceActionType::Button, FInputChord());
 	UI_COMMAND(AddToFavorite, "Add to Favorite", "Add this bookmark to your favorites.", EUserInterfaceActionType::Button, FInputChord());
 	UI_COMMAND(RemoveFromFavorite, "Remove from Favorite", "Remove this bookmark from your favorites.", EUserInterfaceActionType::Button, FInputChord());
+	UI_COMMAND(PlayFromLocation, "Play From Here", "Play from here.", EUserInterfaceActionType::Button, FInputChord());
+	UI_COMMAND(MoveCameraToLocation, "Move Camera Here", "Move the camera to the selected location.", EUserInterfaceActionType::Button, FInputChord());
 }
 
 #undef LOCTEXT_NAMESPACE
diff --git a/Engine/Source/Editor/WorldBookmark/Private/WorldBookmark/WorldBookmarkCommands.h b/Engine/Source/Editor/WorldBookmark/Private/WorldBookmark/WorldBookmarkCommands.h
--- a/Engine/Source/Editor/WorldBookmark/Private/WorldBookmark/WorldBookmarkCommands.h
+++ b/Engine/Source/Editor/WorldBookmark/Private/WorldBookmark/WorldBookmarkCommands.h
@@ -13,7 +13,9 @@ public:
 	TSharedPtr<FUICommandInfo> UpdateBookmark;
 	TSharedPtr<FUICommandInfo> AddToFavorite;
 	TSharedPtr<FUICommandInfo> RemoveFromFavorite;
-		
+	TSharedPtr<FUICommandInfo> PlayFromLocation;
+	TSharedPtr<FUICommandInfo> MoveCameraToLocation;
+	
 	/**
 	 * Initialize commands
 	 */
diff --git a/Engine/Source/Editor/WorldBookmark/Private/WorldBookmark/WorldBookmarkStyle.cpp b/Engine/Source/Editor/WorldBookmark/Private/WorldBookmark/WorldBookmarkStyle.cpp
--- a/Engine/Source/Editor/WorldBookmark/Private/WorldBookmark/WorldBookmarkStyle.cpp
+++ b/Engine/Source/Editor/WorldBookmark/Private/WorldBookmark/WorldBookmarkStyle.cpp
@@ -12,6 +12,7 @@ FWorldBookmarkStyle::FWorldBookmarkStyle()
 {
 	const FVector2D Icon16x16(16.0f, 16.0f);
 	const FVector2D Icon20x20(20.0f, 20.0f);
+	const FVector2D Icon40x40(40.0f, 40.0f);
 	const FVector2D Icon64x64(64.0f, 64.0f);
 
 	SetParentStyleName(FAppStyle::GetAppStyleSetName());
@@ -29,6 +30,8 @@ FWorldBookmarkStyle::FWorldBookmarkStyle()
 	Set("WorldBookmark.CreateBookmark", new CORE_IMAGE_BRUSH_SVG("Starship/Common/plus", Icon16x16, FStyleColors::AccentGreen));
 	Set("WorldBookmark.AddToFavorite", new IMAGE_BRUSH("Icons/Star_16x", Icon16x16));
 	Set("WorldBookmark.RemoveFromFavorite", new IMAGE_BRUSH("Icons/EmptyStar_16x", Icon16x16));
+	Set("WorldBookmark.PlayFromLocation", new IMAGE_BRUSH_SVG("Starship/Common/play", Icon20x20));
+	Set("WorldBookmark.MoveCameraToLocation", new IMAGE_BRUSH_SVG("Starship/EditorViewport/actor-pilot-camera", Icon16x16));
 
 	// Icons
 	Set("WorldBookmark.TabIcon", new IMAGE_BRUSH_SVG("Starship/AssetIcons/WorldBookmark_16", Icon16x16));
diff --git a/Engine/Source/Editor/WorldBookmark/Public/WorldBookmark/WorldBookmark.h b/Engine/Source/Editor/WorldBookmark/Public/WorldBookmark/WorldBookmark.h
--- a/Engine/Source/Editor/WorldBookmark/Public/WorldBookmark/WorldBookmark.h
+++ b/Engine/Source/Editor/WorldBookmark/Public/WorldBookmark/WorldBookmark.h
@@ -27,6 +27,12 @@ class UWorldBookmark : public UObject
 	/** Update the bookmark so that it reflects the current state of the editor. */
 	void Update();
 
+	/** Load the bookmark data, restoring only a specific set of editor states. */
+	void LoadStates(const TArray<TSubclassOf<UEditorState>>& InStatesToLoad);
+
+	/** Update a set of states for the bookmark. */
+	void UpdateStates(const TArray<TSubclassOf<UEditorState>>& InStatesToUpdate);
+
 	/** Return true if this bookmark was flagged as being a favorite bookmark of the user. */
 	bool GetIsUserFavorite() const;
 
@@ -42,6 +48,18 @@ class UWorldBookmark : public UObject
     /** Get the asset registry tag used to store the matching world's name. */
 	static FName GetWorldNameAssetTag();
 
+	template <typename TEditorStateType>
+	typename TEnableIf<TIsDerivedFrom<TEditorStateType, UEditorState>::IsDerived, bool>::Type HasEditorState() const
+	{
+		return EditorState.HasState<TEditorStateType>();
+	}
+
+	template <typename TEditorStateType>
+	typename TEnableIf<TIsDerivedFrom<TEditorStateType, UEditorState>::IsDerived, const TEditorStateType*>::Type GetEditorState() const
+	{
+		return EditorState.GetState<TEditorStateType>();
+	}
+
 private:
 	friend class FWorldBookmarkDetailsCustomization;
 
diff --git a/Engine/Source/Editor/WorldBookmark/WorldBookmark.Build.cs b/Engine/Source/Editor/WorldBookmark/WorldBookmark.Build.cs
--- a/Engine/Source/Editor/WorldBookmark/WorldBookmark.Build.cs
+++ b/Engine/Source/Editor/WorldBookmark/WorldBookmark.Build.cs
@@ -19,6 +19,7 @@ public class WorldBookmark : ModuleRules
 				"EditorConfig",
 				"Engine",
 				"InputCore",
+				"LevelEditor",
 				"PropertyEditor",
 				"SceneOutliner",
 				"SourceControl",
