World Bookmark - Added "World Name" column to the World Bookmarks browser
* Hidden by default
* Added visible columns user settings

diff --git a/Engine/Source/Editor/WorldBookmark/Private/WorldBookmark/Browser/SWorldBookmarkBrowser.cpp b/Engine/Source/Editor/WorldBookmark/Private/WorldBookmark/Browser/SWorldBookmarkBrowser.cpp
--- a/Engine/Source/Editor/WorldBookmark/Private/WorldBookmark/Browser/SWorldBookmarkBrowser.cpp
+++ b/Engine/Source/Editor/WorldBookmark/Private/WorldBookmark/Browser/SWorldBookmarkBrowser.cpp
@@ -39,26 +39,35 @@
 
 namespace WorldBookmarkBrowserColumn
 {
-	namespace Favorite
+	const FColumnDefinition Favorite =
 	{
-		FName Id = TEXT("Favorite");
-		FText DisplayText = FText::GetEmpty();
-		FText ToolTipText = LOCTEXT("Favorite_Tooltip", "Favorite");
+		TEXT("Favorite"),
+		LOCTEXT("Column_Favorite_Text", "Favorite"),
+		LOCTEXT("Column_Favorite_Tooltip", "Favorite")
 	};
 
-	namespace RecentlyUsed
+	const FColumnDefinition RecentlyUsed =
 	{
-		FName Id = TEXT("RecentlyUsed");
-		FText DisplayText = FText::GetEmpty();
-		FText ToolTipText = LOCTEXT("RecentlyUsed_Tooltip", "Recently Used");
+		TEXT("RecentlyUsed"),
+		LOCTEXT("Column_RecentlyUsed_Text", "Recently Used"),
+		LOCTEXT("Column_RecentlyUsed_Tooltip", "Recently Used")
 	};
 
-	namespace Label
+	const FColumnDefinition Label =
 	{
-		FName Id = TEXT("Label");
-		FText DisplayText = LOCTEXT("Label_Text", "Name");
-		FText ToolTipText = LOCTEXT("Label_Tooltip", "World Bookmark name");
+		TEXT("Label"),
+		LOCTEXT("Column_Label_Text", "Name"),
+		LOCTEXT("Column_Label_Tooltip", "World Bookmark name")
 	};
+
+	const FColumnDefinition World =
+	{
+		TEXT("World"),
+		LOCTEXT("Column_World_Text", "World"),
+		LOCTEXT("Column_World_Tooltip", "World Name")
+	};
+
+	const TArray<FName> DefaultHiddenColumns = { World.Id };
 }
 
 static UWorld* GetCurrentWorld()
@@ -92,8 +101,8 @@ SWorldBookmarkBrowser::SWorldBookmarkBrowser()
 	FLevelEditorModule& LevelEditorModule = FModuleManager::LoadModuleChecked<FLevelEditorModule>(TEXT("LevelEditor"));
 	OnMapChangedHandle = LevelEditorModule.OnMapChanged().AddRaw(this, &SWorldBookmarkBrowser::OnMapChanged);
 
-	ColumnsSortingParams[EColumnSortPriority::Primary] = { WorldBookmarkBrowserColumn::Favorite::Id, EColumnSortMode::Ascending };
-	ColumnsSortingParams[EColumnSortPriority::Secondary] = { WorldBookmarkBrowserColumn::Label::Id, EColumnSortMode::Ascending };
+	ColumnsSortingParams[EColumnSortPriority::Primary] = { WorldBookmarkBrowserColumn::Favorite.Id, EColumnSortMode::Ascending };
+	ColumnsSortingParams[EColumnSortPriority::Secondary] = { WorldBookmarkBrowserColumn::Label.Id, EColumnSortMode::Ascending };
 
 	OnSettingsChangedHandle = UWorldBookmarkBrowserSettings::OnSettingsChanged().AddRaw(this, &SWorldBookmarkBrowser::OnSettingsChanged);
 }
@@ -127,6 +136,20 @@ static FSlateIcon GetRecentlyUsedWorldBookmarkIcon()
 	return FSlateIcon(FWorldBookmarkStyle::Get().GetStyleSetName(), RecentlyUsedIconName);
 }
 
+static TSharedRef<SWidget> CreateIconWidget(const FSlateIcon& SlateIcon)
+{
+	return SNew(SBox)
+		.Padding(2.0f)
+		.HAlign(HAlign_Center)
+		.VAlign(VAlign_Center)
+		[
+			SNew(SImage)
+				.DesiredSizeOverride(FVector2D(20, 20))
+				.ColorAndOpacity(FSlateColor::UseForeground())
+				.Image(SlateIcon.GetIcon())
+		];
+}
+
 void SWorldBookmarkBrowser::Construct(const FArguments& InArgs)
 {
 	RefreshItems();
@@ -139,33 +162,46 @@ void SWorldBookmarkBrowser::Construct(const FArguments& InArgs)
 	PropertyView = EditModule.CreateDetailView(DetailsViewArgs);
 	PropertyView->SetObject(nullptr, true);
 
+	auto HeaderRowColumn = [this](const WorldBookmarkBrowserColumn::FColumnDefinition& InColumn) -> SHeaderRow::FColumn::FArguments
+	{
+		return SHeaderRow::Column(InColumn.Id)
+			.SortMode(this, &SWorldBookmarkBrowser::GetColumnSortMode, InColumn.Id)
+			.SortPriority(this, &SWorldBookmarkBrowser::GetColumnSortPriority, InColumn.Id)
+			.OnSort(this, &SWorldBookmarkBrowser::OnColumnSortChanged)
+			.DefaultLabel(InColumn.DisplayText)
+			.DefaultTooltip(InColumn.ToolTipText);
+	};
+
 	TSharedRef<SHeaderRow> HeaderRow = SNew(SHeaderRow)
+		.CanSelectGeneratedColumn(true)
+		.HiddenColumnsList(UWorldBookmarkBrowserSettings::GetHiddenColumns())
+		.OnHiddenColumnsListChanged_Lambda([this]()
+		{
+			UWorldBookmarkBrowserSettings::SetHiddenColumns(BookmarksView->GetHeaderRow()->GetHiddenColumnIds());
+		})
 
 		// Favorite
-		+ SHeaderRow::Column(WorldBookmarkBrowserColumn::Favorite::Id)
-		.SortMode(this, &SWorldBookmarkBrowser::GetColumnSortMode, WorldBookmarkBrowserColumn::Favorite::Id)
-		.SortPriority(this, &SWorldBookmarkBrowser::GetColumnSortPriority, WorldBookmarkBrowserColumn::Favorite::Id)
-		.OnSort(this, &SWorldBookmarkBrowser::OnColumnSortChanged)
+		+ HeaderRowColumn(WorldBookmarkBrowserColumn::Favorite)
 		.FixedWidth(32)
-		.DefaultLabel(WorldBookmarkBrowserColumn::Favorite::DisplayText)
-		.DefaultTooltip(WorldBookmarkBrowserColumn::Favorite::ToolTipText)
+		.HeaderContent()
+		[
+			CreateIconWidget(GetFavoriteWorldBookmarkIcon(true))
+		]
 
 		// RecentlyUsed
-		+ SHeaderRow::Column(WorldBookmarkBrowserColumn::RecentlyUsed::Id)
-		.SortMode(this, &SWorldBookmarkBrowser::GetColumnSortMode, WorldBookmarkBrowserColumn::RecentlyUsed::Id)
-		.SortPriority(this, &SWorldBookmarkBrowser::GetColumnSortPriority, WorldBookmarkBrowserColumn::RecentlyUsed::Id)
-		.OnSort(this, &SWorldBookmarkBrowser::OnColumnSortChanged)
+		+ HeaderRowColumn(WorldBookmarkBrowserColumn::RecentlyUsed)
 		.FixedWidth(32)
-		.DefaultLabel(WorldBookmarkBrowserColumn::RecentlyUsed::DisplayText)
-		.DefaultTooltip(WorldBookmarkBrowserColumn::RecentlyUsed::ToolTipText)
+		.HeaderContent()
+		[
+			CreateIconWidget(GetRecentlyUsedWorldBookmarkIcon())
+		]
 
 		// Label
-		+ SHeaderRow::Column(WorldBookmarkBrowserColumn::Label::Id)
-		.SortMode(this, &SWorldBookmarkBrowser::GetColumnSortMode, WorldBookmarkBrowserColumn::Label::Id)
-		.SortPriority(this, &SWorldBookmarkBrowser::GetColumnSortPriority, WorldBookmarkBrowserColumn::Label::Id)
-		.OnSort(this, &SWorldBookmarkBrowser::OnColumnSortChanged)
-		.DefaultLabel(WorldBookmarkBrowserColumn::Label::DisplayText)
-		.DefaultTooltip(WorldBookmarkBrowserColumn::Label::ToolTipText);
+		+ HeaderRowColumn(WorldBookmarkBrowserColumn::Label)
+		.ShouldGenerateWidget(true) // Ensure the column cannot be hidden (grayed out in the show/hide drop down menu)
+
+		// World
+		+ HeaderRowColumn(WorldBookmarkBrowserColumn::World);
 
 	BookmarksView = SNew(STreeView<FWorldBookmarkTreeItemPtr>)
 		.TreeItemsSource(&WorldBookmarkTreeItems)
@@ -583,18 +619,22 @@ void SWorldBookmarkBrowser::SortItems()
 				return WorldBookmarkA ? -1 : 0;
 			}
 
-			if (SortOrder.Key == WorldBookmarkBrowserColumn::Favorite::Id)
+			if (SortOrder.Key == WorldBookmarkBrowserColumn::Favorite.Id)
 			{
 				return ApplySortMode(WorldBookmarkA->GetIsUserFavorite() == WorldBookmarkB->GetIsUserFavorite() ? 0 : WorldBookmarkA->GetIsUserFavorite() ? -1 : 1);
 			}
-			else if (SortOrder.Key == WorldBookmarkBrowserColumn::RecentlyUsed::Id)
+			else if (SortOrder.Key == WorldBookmarkBrowserColumn::RecentlyUsed.Id)
 			{
 				return ApplySortMode(WorldBookmarkA->GetUserLastLoadedTimeStampUTC() == WorldBookmarkB->GetUserLastLoadedTimeStampUTC() ? 0 : WorldBookmarkA->GetUserLastLoadedTimeStampUTC() < WorldBookmarkB->GetUserLastLoadedTimeStampUTC() ? -1 : 1);
 			}
-			else if (SortOrder.Key == WorldBookmarkBrowserColumn::Label::Id)
+			else if (SortOrder.Key == WorldBookmarkBrowserColumn::Label.Id)
 			{
 				return ApplySortMode(WorldBookmarkA->GetName() == WorldBookmarkB->GetName() ? 0 : WorldBookmarkA->GetName() < WorldBookmarkB->GetName() ? -1 : 1);
 			}
+			else if (SortOrder.Key == WorldBookmarkBrowserColumn::World.Id)
+			{
+				return ApplySortMode(A->GetBookmarkWorld().GetAssetName() == B->GetBookmarkWorld().GetAssetName() ? 0 : A->GetBookmarkWorld().GetAssetName() < B->GetBookmarkWorld().GetAssetName() ? -1 : 1);
+			}
 		
 			checkf(false, TEXT("Unsupported sort order: %s"), *SortOrder.Key.ToString());
 			return 0;
@@ -618,6 +658,16 @@ void SWorldBookmarkBrowser::SortItems()
 	}
 }
 
+const FSoftObjectPath& FWorldBookmarkTreeItem::GetBookmarkWorld() const
+{
+	if (CachedBookmarkWorld.IsNull())
+	{
+		CachedBookmarkWorld = UWorldBookmark::GetWorldFromAssetData(BookmarkAsset);
+	}	
+
+	return CachedBookmarkWorld;
+}
+
 void FWorldBookmarkTableRow::Construct(const FArguments& InArgs, const TSharedRef<STableViewBase>& OwnerTable, FWorldBookmarkTreeItemRef InItem)
 {
 	WeakWorldBookmarkTreeItem = InItem;	
@@ -628,46 +678,28 @@ TSharedRef<SWidget> FWorldBookmarkTableRow::GenerateWidgetForColumn(const FName&
 {
 	if (FWorldBookmarkTreeItemPtr WorldBookmarkTreeItem = WeakWorldBookmarkTreeItem.Pin())
 	{
-		if (ColumnName == WorldBookmarkBrowserColumn::Favorite::Id)
+		if (ColumnName == WorldBookmarkBrowserColumn::Favorite.Id)
 		{
 			if (UWorldBookmark* WorldBookmark = Cast<UWorldBookmark>(WorldBookmarkTreeItem->BookmarkAsset.GetAsset()))
 			{
 				if (WorldBookmark->GetIsUserFavorite())
 				{
-					return SNew(SBox)
-						.Padding(2.0f)
-						.HAlign(HAlign_Left)
-						.VAlign(VAlign_Center)
-						[
-							SNew(SImage)
-								.DesiredSizeOverride(FVector2D(20, 20))
-								.ColorAndOpacity(FSlateColor::UseForeground())
-								.Image(GetFavoriteWorldBookmarkIcon(true).GetIcon())
-						];
+					return CreateIconWidget(GetFavoriteWorldBookmarkIcon(true));
 				}
 			}
 		}
-		else if (ColumnName == WorldBookmarkBrowserColumn::RecentlyUsed::Id)
+		else if (ColumnName == WorldBookmarkBrowserColumn::RecentlyUsed.Id)
 		{
 			if (UWorldBookmark* WorldBookmark = Cast<UWorldBookmark>(WorldBookmarkTreeItem->BookmarkAsset.GetAsset()))
 			{
 				FTimespan LastLoaded = FDateTime::UtcNow() - WorldBookmark->GetUserLastLoadedTimeStampUTC();
 				if (LastLoaded < FTimespan::FromDays(14))
 				{
-					return SNew(SBox)
-						.Padding(2.0f)
-						.HAlign(HAlign_Left)
-						.VAlign(VAlign_Center)
-						[
-							SNew(SImage)
-								.DesiredSizeOverride(FVector2D(20, 20))
-								.ColorAndOpacity(FSlateColor::UseForeground())
-								.Image(GetRecentlyUsedWorldBookmarkIcon().GetIcon())
-						];
+					return CreateIconWidget(GetRecentlyUsedWorldBookmarkIcon());
 				}
 			}
 		}
-		else if (ColumnName == WorldBookmarkBrowserColumn::Label::Id)
+		else if (ColumnName == WorldBookmarkBrowserColumn::Label.Id)
 		{
 			TSharedPtr<SInlineEditableTextBlock> InlineEditableTextBlock;
 
@@ -698,6 +730,17 @@ TSharedRef<SWidget> FWorldBookmarkTableRow::GenerateWidgetForColumn(const FName&
 
 			return LabelWidget;
 		}
+		else if (ColumnName == WorldBookmarkBrowserColumn::World.Id)
+		{
+			return SNew(SBox)
+				.Padding(2.0f)
+				.HAlign(HAlign_Left)
+				.VAlign(VAlign_Center)
+				[
+					SNew(STextBlock)
+						.Text(this, &FWorldBookmarkTableRow::GetWorldNameText)
+				];
+		}
 	}
 
 	return SNullWidget::NullWidget;
@@ -785,6 +828,16 @@ FText FWorldBookmarkTableRow::GetLastAccessedText() const
 	return FText();
 }
 
+FText FWorldBookmarkTableRow::GetWorldNameText() const
+{
+	if (FWorldBookmarkTreeItemPtr WorldBookmarkTreeItem = WeakWorldBookmarkTreeItem.Pin())
+	{
+		return FText::FromString(WorldBookmarkTreeItem->GetBookmarkWorld().GetAssetName());
+	}
+
+	return FText();
+}
+
 void SWorldBookmarkBrowser::OnGetChildren(FWorldBookmarkTreeItemPtr InItem, TArray<FWorldBookmarkTreeItemPtr>& OutChildren)
 {
 }
diff --git a/Engine/Source/Editor/WorldBookmark/Private/WorldBookmark/Browser/SWorldBookmarkBrowser.h b/Engine/Source/Editor/WorldBookmark/Private/WorldBookmark/Browser/SWorldBookmarkBrowser.h
--- a/Engine/Source/Editor/WorldBookmark/Private/WorldBookmark/Browser/SWorldBookmarkBrowser.h
+++ b/Engine/Source/Editor/WorldBookmark/Private/WorldBookmark/Browser/SWorldBookmarkBrowser.h
@@ -14,16 +14,38 @@ class UWorldBookmark;
 class IDetailsView;
 class SAssetView;
 
+namespace WorldBookmarkBrowserColumn
+{
+	struct FColumnDefinition
+	{
+		FName Id;
+		FText DisplayText;
+		FText ToolTipText;
+	};
+
+	extern const FColumnDefinition Favorite;
+	extern const FColumnDefinition RecentlyUsed;
+	extern const FColumnDefinition Label;
+	extern const FColumnDefinition World;
+
+	extern const TArray<FName> DefaultHiddenColumns;
+}
+
 struct FWorldBookmarkTreeItem : TSharedFromThis<FWorldBookmarkTreeItem>
 {
 public:
 	FAssetData BookmarkAsset;
 
+	const FSoftObjectPath& GetBookmarkWorld() const;
+
 	bool bInEditingMode = false;
 
 	/** Delegate for when the context menu requests a rename */
 	DECLARE_DELEGATE(FOnRenameRequested);
 	FOnRenameRequested OnRenameRequested;
+
+private:
+	mutable FSoftObjectPath CachedBookmarkWorld;
 };
 
 typedef TSharedPtr<FWorldBookmarkTreeItem> FWorldBookmarkTreeItemPtr;
@@ -38,6 +60,8 @@ private:
 	virtual TSharedRef<SWidget> GenerateWidgetForColumn(const FName& ColumnName) override;
 
 	FText GetLastAccessedText() const;
+	FText GetWorldNameText() const;
+
 	bool OnLabelTextVerifyChanged(const FText& InNewText, FText& OutErrorMessage);
 	void OnLabelTextCommitted(const FText& NewText, ETextCommit::Type InTextCommit);
 	void OnEnterEditingMode();
@@ -81,7 +105,7 @@ private:
 	EColumnSortMode::Type GetColumnSortMode(const FName ColumnId) const;
 	EColumnSortPriority::Type GetColumnSortPriority(const FName InColumnId) const;
 	TArray<FAssetViewCustomColumn> GetCustomColumns() const;
-
+	
 	void OnAssetsAdded(TConstArrayView<FAssetData> InAssets);
 	void OnAssetRemoved(const FAssetData& InAssetData);
 	void OnAssetRenamed(const FAssetData& InAssetData, const FString& InOldObjectPath);
diff --git a/Engine/Source/Editor/WorldBookmark/Private/WorldBookmark/Browser/WorldBookmarkBrowserSettings.cpp b/Engine/Source/Editor/WorldBookmark/Private/WorldBookmark/Browser/WorldBookmarkBrowserSettings.cpp
--- a/Engine/Source/Editor/WorldBookmark/Private/WorldBookmark/Browser/WorldBookmarkBrowserSettings.cpp
+++ b/Engine/Source/Editor/WorldBookmark/Private/WorldBookmark/Browser/WorldBookmarkBrowserSettings.cpp
@@ -1,6 +1,7 @@
 // Copyright Epic Games, Inc. All Rights Reserved.
 
-#include "WorldBookmarkBrowserSettings.h"
+#include "WorldBookmark/Browser/WorldBookmarkBrowserSettings.h"
+#include "WorldBookmark/Browser/SWorldBookmarkBrowser.h"
 
 TObjectPtr<UWorldBookmarkBrowserSettings> UWorldBookmarkBrowserSettings::Instance = nullptr;
 
@@ -9,6 +10,9 @@ UWorldBookmarkBrowserSettings& UWorldBookmarkBrowserSettings::Get()
 	if (!Instance)
 	{
 		Instance = NewObject<UWorldBookmarkBrowserSettings>();
+
+		Instance->HiddenColumns = WorldBookmarkBrowserColumn::DefaultHiddenColumns;
+
 		Instance->LoadEditorConfig();
 		Instance->AddToRoot();
 	}
@@ -76,6 +80,21 @@ void UWorldBookmarkBrowserSettings::ToggleShowOnlyLastRecentlyUsedBookmarks()
 	Settings.ApplySettingsChanges();
 }
 
+TArray<FName> UWorldBookmarkBrowserSettings::GetHiddenColumns()
+{
+	return Get().HiddenColumns;
+}
+
+void UWorldBookmarkBrowserSettings::SetHiddenColumns(const TArray<FName>& InHiddenColumns)
+{
+	UWorldBookmarkBrowserSettings& Settings = Get();
+	if (Settings.HiddenColumns != InHiddenColumns)
+	{
+		Settings.HiddenColumns = InHiddenColumns;
+		Settings.ApplySettingsChanges();
+	}
+}
+
 int32 UWorldBookmarkBrowserSettings::GetMaxLastRecentlyUsedItems()
 {
 	return Get().MaxLastRecentlyUsedItems;
diff --git a/Engine/Source/Editor/WorldBookmark/Private/WorldBookmark/Browser/WorldBookmarkBrowserSettings.h b/Engine/Source/Editor/WorldBookmark/Private/WorldBookmark/Browser/WorldBookmarkBrowserSettings.h
--- a/Engine/Source/Editor/WorldBookmark/Private/WorldBookmark/Browser/WorldBookmarkBrowserSettings.h
+++ b/Engine/Source/Editor/WorldBookmark/Private/WorldBookmark/Browser/WorldBookmarkBrowserSettings.h
@@ -27,6 +27,9 @@ public:
 	static bool GetShowOnlyLastRecentlyUsedBookmarks();
 	static void ToggleShowOnlyLastRecentlyUsedBookmarks();
 
+	static TArray<FName> GetHiddenColumns();
+	static void SetHiddenColumns(const TArray<FName>& InHiddenColumns);
+
 	static int32 GetMaxLastRecentlyUsedItems();
 	static void SetMaxLastRecentlyUsedItems(int32 InMaxLastRecentlyUsedItems);
 
@@ -47,9 +50,12 @@ private:
 	UPROPERTY(meta=(EditorConfig))
 	bool bShowOnlyLastRecentlyUsedBookmarks = false;
 
-	UPROPERTY(meta = (EditorConfig))
+	UPROPERTY(meta=(EditorConfig))
 	int32 MaxLastRecentlyUsedItems = 5;
 
+	UPROPERTY(meta=(EditorConfig))
+	TArray<FName> HiddenColumns;
+
 	FWorldBookmarkBrowserSettingChanged OnSettingsChangedDelegate;
 
 	static TObjectPtr<UWorldBookmarkBrowserSettings> Instance;
