Missing files in CL 37610213
#cisfix

diff --git a/Engine/Source/Editor/WorldBookmark/Private/WorldBookmark/WorldBookmark.cpp b/Engine/Source/Editor/WorldBookmark/Private/WorldBookmark/WorldBookmark.cpp
--- a/Engine/Source/Editor/WorldBookmark/Private/WorldBookmark/WorldBookmark.cpp
+++ b/Engine/Source/Editor/WorldBookmark/Private/WorldBookmark/WorldBookmark.cpp
@@ -74,6 +74,20 @@ FName UWorldBookmark::GetWorldNameAssetTag()
 	return NAME_WorldName;
 }
 
+FSoftObjectPath UWorldBookmark::GetWorldFromAssetData(const FAssetData& InAssetData)
+{
+	FSoftObjectPath WorldPath;
+
+	FString WorldAssetName;
+	if (InAssetData.GetTagValue(GetWorldNameAssetTag(), WorldAssetName))
+	{
+		WorldPath = FSoftObjectPath(WorldAssetName);
+		UAssetRegistryHelpers::FixupRedirectedAssetPath(WorldPath);
+	}
+
+	return WorldPath;
+}
+
 void UWorldBookmark::GetAssetRegistryTags(FAssetRegistryTagsContext Context) const
 {
 	Super::GetAssetRegistryTags(Context);
diff --git a/Engine/Source/Editor/WorldBookmark/Public/WorldBookmark/WorldBookmark.h b/Engine/Source/Editor/WorldBookmark/Public/WorldBookmark/WorldBookmark.h
--- a/Engine/Source/Editor/WorldBookmark/Public/WorldBookmark/WorldBookmark.h
+++ b/Engine/Source/Editor/WorldBookmark/Public/WorldBookmark/WorldBookmark.h
@@ -48,6 +48,9 @@ class UWorldBookmark : public UObject
     /** Get the asset registry tag used to store the matching world's name. */
 	static FName GetWorldNameAssetTag();
 
+	/** Get the world associated with a bookmark's asset data. */
+	static FSoftObjectPath GetWorldFromAssetData(const FAssetData& InAssetData);
+
 	template <typename TEditorStateType>
 	typename TEnableIf<TIsDerivedFrom<TEditorStateType, UEditorState>::IsDerived, bool>::Type HasEditorState() const
 	{
