WorldBookmark Browser - Moved code to multiple files, added base tree item class (preparing for folder support)

diff --git a/Engine/Source/Editor/WorldBookmark/Private/WorldBookmark/Browser/BookmarkTableRow.cpp b/Engine/Source/Editor/WorldBookmark/Private/WorldBookmark/Browser/BookmarkTableRow.cpp
new file mode 100644
--- /dev/null
+++ b/Engine/Source/Editor/WorldBookmark/Private/WorldBookmark/Browser/BookmarkTableRow.cpp
@@ -0,0 +1,185 @@
+// Copyright Epic Games, Inc. All Rights Reserved.
+
+#include "WorldBookmark/Browser/BookmarkTableRow.h"
+
+#include "AssetToolsModule.h"
+#include "AssetViewUtils.h"
+#include "Modules/ModuleManager.h"
+#include "WorldBookmark/Browser/Columns.h"
+#include "WorldBookmark/Browser/Icons.h"
+#include "WorldBookmark/WorldBookmark.h"
+#include "Widgets/Text/SInlineEditableTextBlock.h"
+#include "Widgets/Text/STextBlock.h"
+#include "Widgets/Views/SListView.h"
+
+#define LOCTEXT_NAMESPACE "WorldBookmarkBrowser"
+
+void FWorldBookmarkTableRow::Construct(const FArguments& InArgs, const TSharedRef<STableViewBase>& OwnerTable, FWorldBookmarkTreeItemRef InItem)
+{
+	WeakWorldBookmarkTreeItem = InItem;	
+	SMultiColumnTableRow<FWorldBookmarkTreeItemPtr>::Construct(InArgs, OwnerTable);
+}
+
+TSharedRef<SWidget> FWorldBookmarkTableRow::GenerateWidgetForColumn(const FName& ColumnName)
+{
+	if (FWorldBookmarkTreeItemPtr WorldBookmarkTreeItem = WeakWorldBookmarkTreeItem.Pin())
+	{
+		if (ColumnName == WorldBookmarkBrowserColumn::Favorite.Id)
+		{
+			if (UWorldBookmark* WorldBookmark = Cast<UWorldBookmark>(WorldBookmarkTreeItem->BookmarkAsset.GetAsset()))
+			{
+				if (WorldBookmark->GetIsUserFavorite())
+				{
+					return CreateIconWidget(GetFavoriteWorldBookmarkIcon(true));
+				}
+			}
+		}
+		else if (ColumnName == WorldBookmarkBrowserColumn::RecentlyUsed.Id)
+		{
+			if (UWorldBookmark* WorldBookmark = Cast<UWorldBookmark>(WorldBookmarkTreeItem->BookmarkAsset.GetAsset()))
+			{
+				FTimespan LastLoaded = FDateTime::UtcNow() - WorldBookmark->GetUserLastLoadedTimeStampUTC();
+				if (LastLoaded < FTimespan::FromDays(14))
+				{
+					return CreateIconWidget(GetRecentlyUsedWorldBookmarkIcon());
+				}
+			}
+		}
+		else if (ColumnName == WorldBookmarkBrowserColumn::Label.Id)
+		{
+			TSharedPtr<SInlineEditableTextBlock> InlineEditableTextBlock;
+
+			TSharedRef<SWidget> LabelWidget =
+				SNew(SVerticalBox)
+				+ SVerticalBox::Slot()
+				.Padding(10, 5, 0, 0)
+				.AutoHeight()
+				[
+					SAssignNew(InlineEditableTextBlock, SInlineEditableTextBlock)
+						.Text(FText::FromName(WorldBookmarkTreeItem->BookmarkAsset.AssetName))
+						.OnVerifyTextChanged(this, &FWorldBookmarkTableRow::OnLabelTextVerifyChanged)
+						.OnTextCommitted(this, &FWorldBookmarkTableRow::OnLabelTextCommitted)
+						.OnEnterEditingMode(this, &FWorldBookmarkTableRow::OnEnterEditingMode)
+						.OnExitEditingMode(this, &FWorldBookmarkTableRow::OnExitEditingMode)
+						.IsSelected(this, &FWorldBookmarkTableRow::IsSelected)
+				]
+				+ SVerticalBox::Slot()
+				.Padding(10, 0, 0, 5)
+				.AutoHeight()
+				[
+					SNew(STextBlock)
+						.ColorAndOpacity(FSlateColor::UseSubduedForeground())
+						.Text(this, &FWorldBookmarkTableRow::GetLastAccessedText)
+				];
+
+			WorldBookmarkTreeItem->OnRenameRequested.BindSP(InlineEditableTextBlock.Get(), &SInlineEditableTextBlock::EnterEditingMode);
+
+			return LabelWidget;
+		}
+		else if (ColumnName == WorldBookmarkBrowserColumn::World.Id)
+		{
+			return SNew(SBox)
+				.Padding(2.0f)
+				.HAlign(HAlign_Left)
+				.VAlign(VAlign_Center)
+				[
+					SNew(STextBlock)
+						.Text(this, &FWorldBookmarkTableRow::GetWorldNameText)
+				];
+		}
+	}
+
+	return SNullWidget::NullWidget;
+}
+
+static FSoftObjectPath GetBookmarkObjectPath(FWorldBookmarkTreeItemPtr InWorldBookmarkTreeItem)
+{
+	return InWorldBookmarkTreeItem->BookmarkAsset.GetSoftObjectPath();
+}
+
+static FSoftObjectPath GetRenamedBookmarkObjectPath(FWorldBookmarkTreeItemPtr InWorldBookmarkTreeItem, const FText& InNewAssetName)
+{
+	const FString PackagePath = InWorldBookmarkTreeItem->BookmarkAsset.PackagePath.ToString();
+	const FString NewObjectPath = FString::Format(TEXT("{0}/{1}.{1}"), { *PackagePath, *InNewAssetName.ToString() });
+	return FSoftObjectPath(NewObjectPath);
+}
+
+bool FWorldBookmarkTableRow::OnLabelTextVerifyChanged(const FText& InNewText, FText& OutErrorMessage)
+{
+	if (FWorldBookmarkTreeItemPtr WorldBookmarkTreeItem = WeakWorldBookmarkTreeItem.Pin())
+	{
+		const FSoftObjectPath NewObjectPath = GetRenamedBookmarkObjectPath(WorldBookmarkTreeItem, InNewText);
+		return AssetViewUtils::IsValidObjectPathForCreate(NewObjectPath.ToString(), OutErrorMessage);
+	}
+
+	return false;
+}
+
+void FWorldBookmarkTableRow::OnLabelTextCommitted(const FText& InNewText, ETextCommit::Type InTextCommit)
+{
+	FText ErrorMessage;
+	if (OnLabelTextVerifyChanged(InNewText, ErrorMessage))
+	{
+		if (FWorldBookmarkTreeItemPtr WorldBookmarkTreeItem = WeakWorldBookmarkTreeItem.Pin())
+		{
+			FAssetToolsModule& AssetToolsModule = FModuleManager::LoadModuleChecked<FAssetToolsModule>("AssetTools");
+
+			TArray<FAssetRenameData> AssetsAndNames;
+			const FSoftObjectPath OldObjectPath = GetBookmarkObjectPath(WorldBookmarkTreeItem);
+			const FSoftObjectPath NewObjectPath = GetRenamedBookmarkObjectPath(WorldBookmarkTreeItem, InNewText);
+			new (AssetsAndNames) FAssetRenameData(OldObjectPath, NewObjectPath);
+
+			AssetToolsModule.Get().RenameAssetsWithDialog(AssetsAndNames);
+		}
+	}
+}
+
+void FWorldBookmarkTableRow::OnEnterEditingMode()
+{
+	if (FWorldBookmarkTreeItemPtr WorldBookmarkTreeItem = WeakWorldBookmarkTreeItem.Pin())
+	{
+		WorldBookmarkTreeItem->bInEditingMode = true;
+	}
+}
+
+void FWorldBookmarkTableRow::OnExitEditingMode()
+{
+	if (FWorldBookmarkTreeItemPtr WorldBookmarkTreeItem = WeakWorldBookmarkTreeItem.Pin())
+	{
+		WorldBookmarkTreeItem->bInEditingMode = false;
+	}
+}
+
+FText FWorldBookmarkTableRow::GetLastAccessedText() const
+{
+	if (FWorldBookmarkTreeItemPtr WorldBookmarkTreeItem = WeakWorldBookmarkTreeItem.Pin())
+	{
+		if (UWorldBookmark* WorldBookmark = Cast<UWorldBookmark>(WorldBookmarkTreeItem->BookmarkAsset.GetAsset()))
+		{
+			FDateTime LastLoadedUTC = WorldBookmark->GetUserLastLoadedTimeStampUTC();
+
+			if (LastLoadedUTC != FDateTime::MinValue())
+			{
+				FTimespan UTCOffset = FDateTime::UtcNow() - FDateTime::Now();
+				uint64 AdjustedTicks = LastLoadedUTC.GetTicks() - UTCOffset.GetTicks();
+				FDateTime LastLoadedLocal = FDateTime(AdjustedTicks);
+
+				return FText::Format(LOCTEXT("LastAccessedDateTime", "Last accessed {0}"), FText::AsDateTime(LastLoadedUTC, TEXT("%Y-%m-%d %H:%M")));
+			}
+		}
+	}
+
+	return FText();
+}
+
+FText FWorldBookmarkTableRow::GetWorldNameText() const
+{
+	if (FWorldBookmarkTreeItemPtr WorldBookmarkTreeItem = WeakWorldBookmarkTreeItem.Pin())
+	{
+		return FText::FromString(WorldBookmarkTreeItem->GetBookmarkWorld().GetAssetName());
+	}
+
+	return FText();
+}
+
+#undef LOCTEXT_NAMESPACE
diff --git a/Engine/Source/Editor/WorldBookmark/Private/WorldBookmark/Browser/BookmarkTableRow.h b/Engine/Source/Editor/WorldBookmark/Private/WorldBookmark/Browser/BookmarkTableRow.h
new file mode 100644
--- /dev/null
+++ b/Engine/Source/Editor/WorldBookmark/Private/WorldBookmark/Browser/BookmarkTableRow.h
@@ -0,0 +1,28 @@
+// Copyright Epic Games, Inc. All Rights Reserved.
+#pragma once
+
+#include "Widgets/Views/STableRow.h"
+#include "Widgets/Views/SListView.h"
+#include "WorldBookmark/Browser/BookmarkTreeItem.h"
+
+class SWidget;
+
+class FWorldBookmarkTableRow : public SMultiColumnTableRow<FWorldBookmarkTreeItemPtr>
+{
+public:
+	void Construct(const FArguments& InArgs, const TSharedRef<STableViewBase>& OwnerTable, FWorldBookmarkTreeItemRef InItem);
+
+private:
+	virtual TSharedRef<SWidget> GenerateWidgetForColumn(const FName& ColumnName) override;
+
+	FText GetLastAccessedText() const;
+	FText GetWorldNameText() const;
+
+	bool OnLabelTextVerifyChanged(const FText& InNewText, FText& OutErrorMessage);
+	void OnLabelTextCommitted(const FText& NewText, ETextCommit::Type InTextCommit);
+	void OnEnterEditingMode();
+	void OnExitEditingMode();
+
+private:
+	TWeakPtr<FWorldBookmarkTreeItem> WeakWorldBookmarkTreeItem;
+};
diff --git a/Engine/Source/Editor/WorldBookmark/Private/WorldBookmark/Browser/BookmarkTreeItem.cpp b/Engine/Source/Editor/WorldBookmark/Private/WorldBookmark/Browser/BookmarkTreeItem.cpp
new file mode 100644
--- /dev/null
+++ b/Engine/Source/Editor/WorldBookmark/Private/WorldBookmark/Browser/BookmarkTreeItem.cpp
@@ -0,0 +1,5 @@
+// Copyright Epic Games, Inc. All Rights Reserved.
+
+#include "WorldBookmark/Browser/BookmarkTreeItem.h"
+
+const ETreeItemType FWorldBookmarkTreeItem::Type(ETreeItemType::WorldBookmark);
diff --git a/Engine/Source/Editor/WorldBookmark/Private/WorldBookmark/Browser/BookmarkTreeItem.h b/Engine/Source/Editor/WorldBookmark/Private/WorldBookmark/Browser/BookmarkTreeItem.h
new file mode 100644
--- /dev/null
+++ b/Engine/Source/Editor/WorldBookmark/Private/WorldBookmark/Browser/BookmarkTreeItem.h
@@ -0,0 +1,35 @@
+// Copyright Epic Games, Inc. All Rights Reserved.
+#pragma once
+
+#include "WorldBookmark/Browser/TreeItem.h"
+
+#include "AssetRegistry/AssetData.h"
+
+struct FWorldBookmarkTreeItem : public ITreeItem
+{
+public:
+	FWorldBookmarkTreeItem(const FAssetData& InBookmarkAsset)
+		: ITreeItem(Type)
+		, BookmarkAsset(InBookmarkAsset)
+	{
+	}
+
+	/** Static type identifier for this tree item class */
+	static const ETreeItemType Type;
+
+	FAssetData BookmarkAsset;
+
+	virtual FString GetName() const override
+	{
+		return BookmarkAsset.AssetName.ToString();
+	}
+
+	const FSoftObjectPath& GetBookmarkWorld() const;
+
+private:
+
+	mutable FSoftObjectPath CachedBookmarkWorld;
+};
+
+typedef TSharedPtr<FWorldBookmarkTreeItem> FWorldBookmarkTreeItemPtr;
+typedef TSharedRef<FWorldBookmarkTreeItem> FWorldBookmarkTreeItemRef;
diff --git a/Engine/Source/Editor/WorldBookmark/Private/WorldBookmark/Browser/Columns.cpp b/Engine/Source/Editor/WorldBookmark/Private/WorldBookmark/Browser/Columns.cpp
new file mode 100644
--- /dev/null
+++ b/Engine/Source/Editor/WorldBookmark/Private/WorldBookmark/Browser/Columns.cpp
@@ -0,0 +1,42 @@
+// Copyright Epic Games, Inc. All Rights Reserved.
+
+#include "WorldBookmark/Browser/Columns.h"
+
+#include "Internationalization/Internationalization.h"
+
+#define LOCTEXT_NAMESPACE "WorldBookmarkBrowser"
+
+namespace WorldBookmarkBrowserColumn
+{
+	const FColumnDefinition Favorite =
+	{
+		TEXT("Favorite"),
+		LOCTEXT("Column_Favorite_Text", "Favorite"),
+		LOCTEXT("Column_Favorite_Tooltip", "Favorite")
+	};
+
+	const FColumnDefinition RecentlyUsed =
+	{
+		TEXT("RecentlyUsed"),
+		LOCTEXT("Column_RecentlyUsed_Text", "Recently Used"),
+		LOCTEXT("Column_RecentlyUsed_Tooltip", "Recently Used")
+	};
+
+	const FColumnDefinition Label =
+	{
+		TEXT("Label"),
+		LOCTEXT("Column_Label_Text", "Name"),
+		LOCTEXT("Column_Label_Tooltip", "World Bookmark name")
+	};
+
+	const FColumnDefinition World =
+	{
+		TEXT("World"),
+		LOCTEXT("Column_World_Text", "World"),
+		LOCTEXT("Column_World_Tooltip", "World Name")
+	};
+
+	const TArray<FName> DefaultHiddenColumns = { World.Id };
+}
+
+#undef LOCTEXT_NAMESPACE
diff --git a/Engine/Source/Editor/WorldBookmark/Private/WorldBookmark/Browser/Columns.h b/Engine/Source/Editor/WorldBookmark/Private/WorldBookmark/Browser/Columns.h
new file mode 100644
--- /dev/null
+++ b/Engine/Source/Editor/WorldBookmark/Private/WorldBookmark/Browser/Columns.h
@@ -0,0 +1,23 @@
+// Copyright Epic Games, Inc. All Rights Reserved.
+#pragma once
+
+#include "Containers/Array.h"
+#include "Internationalization/Text.h"
+#include "UObject/NameTypes.h"
+
+namespace WorldBookmarkBrowserColumn
+{
+	struct FColumnDefinition
+	{
+		FName Id;
+		FText DisplayText;
+		FText ToolTipText;
+	};
+
+	extern const FColumnDefinition Favorite;
+	extern const FColumnDefinition RecentlyUsed;
+	extern const FColumnDefinition Label;
+	extern const FColumnDefinition World;
+
+	extern const TArray<FName> DefaultHiddenColumns;
+}
diff --git a/Engine/Source/Editor/WorldBookmark/Private/WorldBookmark/Browser/Icons.cpp b/Engine/Source/Editor/WorldBookmark/Private/WorldBookmark/Browser/Icons.cpp
new file mode 100644
--- /dev/null
+++ b/Engine/Source/Editor/WorldBookmark/Private/WorldBookmark/Browser/Icons.cpp
@@ -0,0 +1,33 @@
+// Copyright Epic Games, Inc. All Rights Reserved.
+
+#include "WorldBookmark/Browser/Icons.h"
+
+#include "Widgets/Images/SImage.h"
+#include "Widgets/Layout/SBox.h"
+#include "WorldBookmark/WorldBookmarkStyle.h"
+
+FSlateIcon GetFavoriteWorldBookmarkIcon(bool bIsFavorite)
+{
+	const FName FavoriteIconName = bIsFavorite ? "WorldBookmark.IsFavorite" : "WorldBookmark.IsNotFavorite";
+	return FSlateIcon(FWorldBookmarkStyle::Get().GetStyleSetName(), FavoriteIconName);
+}
+
+FSlateIcon GetRecentlyUsedWorldBookmarkIcon()
+{
+	const FName RecentlyUsedIconName = "WorldBookmark.RecentlyUsed";
+	return FSlateIcon(FWorldBookmarkStyle::Get().GetStyleSetName(), RecentlyUsedIconName);
+}
+
+TSharedRef<SWidget> CreateIconWidget(const FSlateIcon& SlateIcon)
+{
+	return SNew(SBox)
+		.Padding(2.0f)
+		.HAlign(HAlign_Center)
+		.VAlign(VAlign_Center)
+		[
+			SNew(SImage)
+				.DesiredSizeOverride(FVector2D(20, 20))
+				.ColorAndOpacity(FSlateColor::UseForeground())
+				.Image(SlateIcon.GetIcon())
+		];
+}
diff --git a/Engine/Source/Editor/WorldBookmark/Private/WorldBookmark/Browser/Icons.h b/Engine/Source/Editor/WorldBookmark/Private/WorldBookmark/Browser/Icons.h
new file mode 100644
--- /dev/null
+++ b/Engine/Source/Editor/WorldBookmark/Private/WorldBookmark/Browser/Icons.h
@@ -0,0 +1,11 @@
+// Copyright Epic Games, Inc. All Rights Reserved.
+#pragma once
+
+#include "Templates/SharedPointer.h"
+#include "Textures/SlateIcon.h"
+
+class SWidget;
+
+FSlateIcon GetFavoriteWorldBookmarkIcon(bool bIsFavorite);
+FSlateIcon GetRecentlyUsedWorldBookmarkIcon();
+TSharedRef<SWidget> CreateIconWidget(const FSlateIcon& SlateIcon);
diff --git a/Engine/Source/Editor/WorldBookmark/Private/WorldBookmark/Browser/SWorldBookmarkBrowser.cpp b/Engine/Source/Editor/WorldBookmark/Private/WorldBookmark/Browser/SWorldBookmarkBrowser.cpp
--- a/Engine/Source/Editor/WorldBookmark/Private/WorldBookmark/Browser/SWorldBookmarkBrowser.cpp
+++ b/Engine/Source/Editor/WorldBookmark/Private/WorldBookmark/Browser/SWorldBookmarkBrowser.cpp
@@ -1,12 +1,23 @@
 // Copyright Epic Games, Inc. All Rights Reserved.
+
 #include "WorldBookmark/Browser/SWorldBookmarkBrowser.h"
 
+#include "WorldBookmark/Browser/Columns.h"
+#include "WorldBookmark/Browser/Icons.h"
+#include "WorldBookmark/Browser/Settings.h"
+#include "WorldBookmark/Browser/BookmarkTreeItem.h"
+#include "WorldBookmark/Browser/BookmarkTableRow.h"
+
+#include "WorldBookmark/WorldBookmark.h"
+#include "WorldBookmark/WorldBookmarkCommands.h"
+#include "WorldBookmark/WorldBookmarkFactory.h"
+#include "WorldBookmark/WorldBookmarkStyle.h"
+
 #include "Algo/AnyOf.h"
 #include "AssetRegistry/IAssetRegistry.h"
 #include "AssetRegistry/AssetData.h"
 #include "AssetRegistry/AssetRegistryHelpers.h"
 #include "AssetToolsModule.h"
-#include "AssetViewUtils.h"
 #include "ContentBrowserDataSource.h"
 #include "ContentBrowserModule.h"
 #include "Editor.h"
@@ -29,47 +40,10 @@
 #include "Widgets/Input/SSearchBox.h"
 #include "Widgets/Input/SSpinBox.h"
 #include "Widgets/Text/SInlineEditableTextBlock.h"
-#include "WorldBookmark/WorldBookmark.h"
-#include "WorldBookmark/WorldBookmarkCommands.h"
-#include "WorldBookmark/WorldBookmarkFactory.h"
-#include "WorldBookmark/WorldBookmarkStyle.h"
-#include "WorldBookmark/Browser/WorldBookmarkBrowserSettings.h"
+#include "Widgets/Views/SHeaderRow.h"
 
 #define LOCTEXT_NAMESPACE "WorldBookmarkBrowser"
 
-namespace WorldBookmarkBrowserColumn
-{
-	const FColumnDefinition Favorite =
-	{
-		TEXT("Favorite"),
-		LOCTEXT("Column_Favorite_Text", "Favorite"),
-		LOCTEXT("Column_Favorite_Tooltip", "Favorite")
-	};
-
-	const FColumnDefinition RecentlyUsed =
-	{
-		TEXT("RecentlyUsed"),
-		LOCTEXT("Column_RecentlyUsed_Text", "Recently Used"),
-		LOCTEXT("Column_RecentlyUsed_Tooltip", "Recently Used")
-	};
-
-	const FColumnDefinition Label =
-	{
-		TEXT("Label"),
-		LOCTEXT("Column_Label_Text", "Name"),
-		LOCTEXT("Column_Label_Tooltip", "World Bookmark name")
-	};
-
-	const FColumnDefinition World =
-	{
-		TEXT("World"),
-		LOCTEXT("Column_World_Text", "World"),
-		LOCTEXT("Column_World_Tooltip", "World Name")
-	};
-
-	const TArray<FName> DefaultHiddenColumns = { World.Id };
-}
-
 static UWorld* GetCurrentWorld()
 {
 	if (GEditor)
@@ -124,32 +98,6 @@ SWorldBookmarkBrowser::~SWorldBookmarkBrowser()
 	UWorldBookmarkBrowserSettings::OnSettingsChanged().Remove(OnSettingsChangedHandle);
 }
 
-static FSlateIcon GetFavoriteWorldBookmarkIcon(bool bIsFavorite)
-{
-	const FName FavoriteIconName = bIsFavorite ? "WorldBookmark.IsFavorite" : "WorldBookmark.IsNotFavorite";
-	return FSlateIcon(FWorldBookmarkStyle::Get().GetStyleSetName(), FavoriteIconName);
-}
-
-static FSlateIcon GetRecentlyUsedWorldBookmarkIcon()
-{
-	const FName RecentlyUsedIconName = "WorldBookmark.RecentlyUsed";
-	return FSlateIcon(FWorldBookmarkStyle::Get().GetStyleSetName(), RecentlyUsedIconName);
-}
-
-static TSharedRef<SWidget> CreateIconWidget(const FSlateIcon& SlateIcon)
-{
-	return SNew(SBox)
-		.Padding(2.0f)
-		.HAlign(HAlign_Center)
-		.VAlign(VAlign_Center)
-		[
-			SNew(SImage)
-				.DesiredSizeOverride(FVector2D(20, 20))
-				.ColorAndOpacity(FSlateColor::UseForeground())
-				.Image(SlateIcon.GetIcon())
-		];
-}
-
 void SWorldBookmarkBrowser::Construct(const FArguments& InArgs)
 {
 	RefreshItems();
@@ -203,15 +151,15 @@ void SWorldBookmarkBrowser::Construct(const FArguments& InArgs)
 		// World
 		+ HeaderRowColumn(WorldBookmarkBrowserColumn::World);
 
-	BookmarksView = SNew(STreeView<FWorldBookmarkTreeItemPtr>)
-		.TreeItemsSource(&WorldBookmarkTreeItems)
+	BookmarksView = SNew(STreeView<FTreeItemPtr>)
+		.TreeItemsSource(&TreeItems)
 		.SelectionMode(ESelectionMode::Single)
 		.OnGetChildren(this, &SWorldBookmarkBrowser::OnGetChildren)
 		.OnGenerateRow(this, &SWorldBookmarkBrowser::OnGenerateRow)
-		.OnSelectionChanged(this, &SWorldBookmarkBrowser::OnSelectedBookmarkChanged)
-		.OnMouseButtonDoubleClick(this, &SWorldBookmarkBrowser::OnBookmarkDoubleClicked)
+		.OnSelectionChanged(this, &SWorldBookmarkBrowser::OnSelectionChanged)
+		.OnMouseButtonDoubleClick(this, &SWorldBookmarkBrowser::OnMouseButtonDoubleClicked)
 		.OnContextMenuOpening(this, &SWorldBookmarkBrowser::OnContextMenuOpening)
-		.OnItemScrolledIntoView(this, &SWorldBookmarkBrowser::ItemScrolledIntoView)
+		.OnItemScrolledIntoView(this, &SWorldBookmarkBrowser::OnItemScrolledIntoView)
 		.HeaderRow
 		(
 			HeaderRow
@@ -421,13 +369,13 @@ TSharedRef<SWidget> SWorldBookmarkBrowser::GetSettingsMenuContent()
 
 FReply SWorldBookmarkBrowser::OnKeyDown(const FGeometry& MyGeometry, const FKeyEvent& InKeyEvent)
 {
-	bool bIsRenamingAsset = false;
-	if (FWorldBookmarkTreeItemPtr SelectedBookmarkTreeItem = GetSelectedBookmarkTreeItem())
+	bool bIsRenamingItem = false;
+	if (FTreeItemPtr SelectedTreeItem = GetSelectedTreeItem())
 	{
-		bIsRenamingAsset = SelectedBookmarkTreeItem->bInEditingMode;
+		bIsRenamingItem = SelectedTreeItem->bInEditingMode;
 	}
 		
-	if (bIsRenamingAsset || Commands->ProcessCommandBindings(InKeyEvent))
+	if (bIsRenamingItem || Commands->ProcessCommandBindings(InKeyEvent))
 	{
 		return FReply::Handled();
 	}
@@ -450,7 +398,7 @@ void SWorldBookmarkBrowser::RefreshItems()
 
 	UWorldBookmark* PreviousSelection = GetSelectedBookmark();
 
-	WorldBookmarkTreeItems.Reset();
+	TreeItems.Reset();
 
 	FSoftObjectPath CurrentWorldPath = FSoftObjectPath(GetCurrentWorld());
 	
@@ -586,8 +534,7 @@ void SWorldBookmarkBrowser::RefreshItems()
 
 	for (const FAssetData& AssetData : AssetsData)
 	{
-		FWorldBookmarkTreeItemPtr& Item = WorldBookmarkTreeItems.Emplace_GetRef(MakeShared<FWorldBookmarkTreeItem>());
-		Item->BookmarkAsset = AssetData;
+		TreeItems.Emplace_GetRef(MakeShared<FWorldBookmarkTreeItem>(AssetData));
 	}
 
 	SortItems();	
@@ -600,8 +547,11 @@ void SWorldBookmarkBrowser::RefreshItems()
 
 void SWorldBookmarkBrowser::SortItems()
 {
-	WorldBookmarkTreeItems.Sort([this](const FWorldBookmarkTreeItemPtr& A, const FWorldBookmarkTreeItemPtr& B)
+	TreeItems.Sort([this](const FTreeItemPtr& LHS, const FTreeItemPtr& RHS)
 	{
+		const FWorldBookmarkTreeItemPtr& A = StaticCastSharedPtr<FWorldBookmarkTreeItem>(LHS);
+		const FWorldBookmarkTreeItemPtr& B = StaticCastSharedPtr<FWorldBookmarkTreeItem>(RHS);
+
 		auto CompareItems = [&](const TPair<FName, EColumnSortMode::Type>& SortOrder) -> int32
 		{
 			auto ApplySortMode = [&](int32 SortResult) { return SortOrder.Value == EColumnSortMode::Ascending ? SortResult : -SortResult; };
@@ -668,211 +618,34 @@ const FSoftObjectPath& FWorldBookmarkTreeItem::GetBookmarkWorld() const
 	return CachedBookmarkWorld;
 }
 
-void FWorldBookmarkTableRow::Construct(const FArguments& InArgs, const TSharedRef<STableViewBase>& OwnerTable, FWorldBookmarkTreeItemRef InItem)
+void SWorldBookmarkBrowser::OnGetChildren(FTreeItemPtr InItem, TArray<FTreeItemPtr>& OutChildren)
 {
-	WeakWorldBookmarkTreeItem = InItem;	
-	SMultiColumnTableRow<FWorldBookmarkTreeItemPtr>::Construct(InArgs, OwnerTable);
 }
 
-TSharedRef<SWidget> FWorldBookmarkTableRow::GenerateWidgetForColumn(const FName& ColumnName)
+TSharedRef<ITableRow> SWorldBookmarkBrowser::OnGenerateRow(FTreeItemPtr InTreeItem, const TSharedRef<STableViewBase>& OwnerTable)
 {
-	if (FWorldBookmarkTreeItemPtr WorldBookmarkTreeItem = WeakWorldBookmarkTreeItem.Pin())
+	if (InTreeItem->IsA<FWorldBookmarkTreeItem>())
 	{
-		if (ColumnName == WorldBookmarkBrowserColumn::Favorite.Id)
-		{
-			if (UWorldBookmark* WorldBookmark = Cast<UWorldBookmark>(WorldBookmarkTreeItem->BookmarkAsset.GetAsset()))
-			{
-				if (WorldBookmark->GetIsUserFavorite())
-				{
-					return CreateIconWidget(GetFavoriteWorldBookmarkIcon(true));
-				}
-			}
-		}
-		else if (ColumnName == WorldBookmarkBrowserColumn::RecentlyUsed.Id)
-		{
-			if (UWorldBookmark* WorldBookmark = Cast<UWorldBookmark>(WorldBookmarkTreeItem->BookmarkAsset.GetAsset()))
-			{
-				FTimespan LastLoaded = FDateTime::UtcNow() - WorldBookmark->GetUserLastLoadedTimeStampUTC();
-				if (LastLoaded < FTimespan::FromDays(14))
-				{
-					return CreateIconWidget(GetRecentlyUsedWorldBookmarkIcon());
-				}
-			}
-		}
-		else if (ColumnName == WorldBookmarkBrowserColumn::Label.Id)
-		{
-			TSharedPtr<SInlineEditableTextBlock> InlineEditableTextBlock;
-
-			TSharedRef<SWidget> LabelWidget =
-				SNew(SVerticalBox)
-				+ SVerticalBox::Slot()
-				.Padding(10, 5, 0, 0)
-				.AutoHeight()
-				[
-					SAssignNew(InlineEditableTextBlock, SInlineEditableTextBlock)
-						.Text(FText::FromName(WorldBookmarkTreeItem->BookmarkAsset.AssetName))
-						.OnVerifyTextChanged(this, &FWorldBookmarkTableRow::OnLabelTextVerifyChanged)
-						.OnTextCommitted(this, &FWorldBookmarkTableRow::OnLabelTextCommitted)
-						.OnEnterEditingMode(this, &FWorldBookmarkTableRow::OnEnterEditingMode)
-						.OnExitEditingMode(this, &FWorldBookmarkTableRow::OnExitEditingMode)
-						.IsSelected(this, &FWorldBookmarkTableRow::IsSelected)
-				]
-				+ SVerticalBox::Slot()
-				.Padding(10, 0, 0, 5)
-				.AutoHeight()
-				[
-					SNew(STextBlock)
-						.ColorAndOpacity(FSlateColor::UseSubduedForeground())
-						.Text(this, &FWorldBookmarkTableRow::GetLastAccessedText)
-				];
-
-			WorldBookmarkTreeItem->OnRenameRequested.BindSP(InlineEditableTextBlock.Get(), &SInlineEditableTextBlock::EnterEditingMode);
-
-			return LabelWidget;
-		}
-		else if (ColumnName == WorldBookmarkBrowserColumn::World.Id)
-		{
-			return SNew(SBox)
-				.Padding(2.0f)
-				.HAlign(HAlign_Left)
-				.VAlign(VAlign_Center)
-				[
-					SNew(STextBlock)
-						.Text(this, &FWorldBookmarkTableRow::GetWorldNameText)
-				];
-		}
+		return SNew(FWorldBookmarkTableRow, OwnerTable, StaticCastSharedPtr<FWorldBookmarkTreeItem>(InTreeItem).ToSharedRef());
 	}
 
-	return SNullWidget::NullWidget;
+	return SNew(STableRow<TSharedPtr<FString>>, OwnerTable);
 }
 
-static FSoftObjectPath GetBookmarkObjectPath(FWorldBookmarkTreeItemPtr InWorldBookmarkTreeItem)
+void SWorldBookmarkBrowser::OnSelectionChanged(FTreeItemPtr InTreeItem, ESelectInfo::Type SelectionType)
 {
-	return InWorldBookmarkTreeItem->BookmarkAsset.GetSoftObjectPath();
-}
-
-static FSoftObjectPath GetRenamedBookmarkObjectPath(FWorldBookmarkTreeItemPtr InWorldBookmarkTreeItem, const FText& InNewAssetName)
-{
-	const FString PackagePath = InWorldBookmarkTreeItem->BookmarkAsset.PackagePath.ToString();
-	const FString NewObjectPath = FString::Format(TEXT("{0}/{1}.{1}"), { *PackagePath, *InNewAssetName.ToString() });
-	return FSoftObjectPath(NewObjectPath);
-}
-
-bool FWorldBookmarkTableRow::OnLabelTextVerifyChanged(const FText& InNewText, FText& OutErrorMessage)
-{
-	if (FWorldBookmarkTreeItemPtr WorldBookmarkTreeItem = WeakWorldBookmarkTreeItem.Pin())
-	{
-		const FSoftObjectPath NewObjectPath = GetRenamedBookmarkObjectPath(WorldBookmarkTreeItem, InNewText);
-		return AssetViewUtils::IsValidObjectPathForCreate(NewObjectPath.ToString(), OutErrorMessage);
-	}
-	else
-	{
-		return false;
-	}
-}
-
-void FWorldBookmarkTableRow::OnLabelTextCommitted(const FText& InNewText, ETextCommit::Type InTextCommit)
-{
-	FText ErrorMessage;
-	if (OnLabelTextVerifyChanged(InNewText, ErrorMessage))
-	{
-		if (FWorldBookmarkTreeItemPtr WorldBookmarkTreeItem = WeakWorldBookmarkTreeItem.Pin())
-		{
-			FAssetToolsModule& AssetToolsModule = FModuleManager::LoadModuleChecked<FAssetToolsModule>("AssetTools");
-
-			TArray<FAssetRenameData> AssetsAndNames;
-			const FSoftObjectPath OldObjectPath = GetBookmarkObjectPath(WorldBookmarkTreeItem);
-			const FSoftObjectPath NewObjectPath = GetRenamedBookmarkObjectPath(WorldBookmarkTreeItem, InNewText);
-			new (AssetsAndNames) FAssetRenameData(OldObjectPath, NewObjectPath);
-
-			AssetToolsModule.Get().RenameAssetsWithDialog(AssetsAndNames);
-		}
-	}
-}
-
-void FWorldBookmarkTableRow::OnEnterEditingMode()
-{
-	if (FWorldBookmarkTreeItemPtr WorldBookmarkTreeItem = WeakWorldBookmarkTreeItem.Pin())
-	{
-		WorldBookmarkTreeItem->bInEditingMode = true;
-	}
-}
-
-void FWorldBookmarkTableRow::OnExitEditingMode()
-{
-	if (FWorldBookmarkTreeItemPtr WorldBookmarkTreeItem = WeakWorldBookmarkTreeItem.Pin())
-	{
-		WorldBookmarkTreeItem->bInEditingMode = false;
-	}
-}
-
-FText FWorldBookmarkTableRow::GetLastAccessedText() const
-{
-	if (FWorldBookmarkTreeItemPtr WorldBookmarkTreeItem = WeakWorldBookmarkTreeItem.Pin())
-	{
-		if (UWorldBookmark* WorldBookmark = Cast<UWorldBookmark>(WorldBookmarkTreeItem->BookmarkAsset.GetAsset()))
-		{
-			FDateTime LastLoadedUTC = WorldBookmark->GetUserLastLoadedTimeStampUTC();
-
-			if (LastLoadedUTC != FDateTime::MinValue())
-			{
-				FTimespan UTCOffset = FDateTime::UtcNow() - FDateTime::Now();
-				uint64 AdjustedTicks = LastLoadedUTC.GetTicks() - UTCOffset.GetTicks();
-				FDateTime LastLoadedLocal = FDateTime(AdjustedTicks);
-
-				return FText::Format(LOCTEXT("LastAccessedDateTime", "Last accessed {0}"), FText::AsDateTime(LastLoadedUTC, TEXT("%Y-%m-%d %H:%M")));
-			}
-		}
-	}
-
-	return FText();
-}
-
-FText FWorldBookmarkTableRow::GetWorldNameText() const
-{
-	if (FWorldBookmarkTreeItemPtr WorldBookmarkTreeItem = WeakWorldBookmarkTreeItem.Pin())
-	{
-		return FText::FromString(WorldBookmarkTreeItem->GetBookmarkWorld().GetAssetName());
-	}
-
-	return FText();
-}
-
-void SWorldBookmarkBrowser::OnGetChildren(FWorldBookmarkTreeItemPtr InItem, TArray<FWorldBookmarkTreeItemPtr>& OutChildren)
-{
-}
-
-TSharedRef<ITableRow> SWorldBookmarkBrowser::OnGenerateRow(FWorldBookmarkTreeItemPtr InTreeItem, const TSharedRef<STableViewBase>& OwnerTable)
-{
-	return SNew(FWorldBookmarkTableRow, OwnerTable, InTreeItem.ToSharedRef());
-}
-
-void SWorldBookmarkBrowser::OnSelectedBookmarkChanged(FWorldBookmarkTreeItemPtr InTreeItem, ESelectInfo::Type SelectionType) const
-{
-	UWorldBookmark* SelectedWorldBookmark = nullptr;
-
-	if (InTreeItem.IsValid() && InTreeItem->BookmarkAsset.IsValid())
-	{
-		if (UWorldBookmark* WorldBookmark = Cast<UWorldBookmark>(InTreeItem->BookmarkAsset.GetAsset()))
-		{
-			SelectedWorldBookmark = WorldBookmark;
-		}
-	}
-
+	UWorldBookmark* SelectedWorldBookmark = GetSelectedBookmark();
 	PropertyView->SetObject(SelectedWorldBookmark);
 }
 
-void SWorldBookmarkBrowser::OnBookmarkDoubleClicked(FWorldBookmarkTreeItemPtr InTreeItem)
+void SWorldBookmarkBrowser::OnMouseButtonDoubleClicked(FTreeItemPtr InTreeItem)
 {
-	if (InTreeItem.IsValid() && InTreeItem->BookmarkAsset.IsValid())
+	if (UWorldBookmark* WorldBookmark = GetSelectedBookmark())
 	{
-		if (UWorldBookmark* WorldBookmark = Cast<UWorldBookmark>(InTreeItem->BookmarkAsset.GetAsset()))
-		{
-			WorldBookmark->Load();
+		WorldBookmark->Load();
 
-			// Loading an item will update it's last load time - refresh the UI
-			RefreshItems();
-		}
+		// Loading an item will update it's last load time - refresh the UI
+		RefreshItems();
 	}
 }
 
@@ -1167,26 +940,37 @@ bool SWorldBookmarkBrowser::IsValidBookmarkSelected() const
 	return GetSelectedBookmark() != nullptr;
 }
 
-FWorldBookmarkTreeItemPtr SWorldBookmarkBrowser::GetSelectedBookmarkTreeItem() const
+FTreeItemPtr SWorldBookmarkBrowser::GetSelectedTreeItem() const
 {
 	if (!BookmarksView)
 	{
 		return nullptr;
 	}
 
-	const TArray<FWorldBookmarkTreeItemPtr> Selection = BookmarksView->GetSelectedItems();
+	const TArray<FTreeItemPtr> Selection = BookmarksView->GetSelectedItems();
 	if (Selection.Num() != 1)
 	{
 		return nullptr;
 	}
 
-	FWorldBookmarkTreeItemPtr SelectedItem = Selection[0];
-	if (!SelectedItem || !SelectedItem->BookmarkAsset.IsValid())
+	return Selection[0];
+}
+
+FWorldBookmarkTreeItemPtr SWorldBookmarkBrowser::GetSelectedBookmarkTreeItem() const
+{
+	FTreeItemPtr SelectedTreeItem = GetSelectedTreeItem();
+	if (!SelectedTreeItem || !SelectedTreeItem->IsA<FWorldBookmarkTreeItem>())
+	{
+		return nullptr;
+	}
+
+	FWorldBookmarkTreeItemPtr SelectedWorldBookmarkTreeItem = StaticCastSharedPtr<FWorldBookmarkTreeItem>(SelectedTreeItem);
+	if (!SelectedWorldBookmarkTreeItem->BookmarkAsset.IsValid())
 	{
 		return nullptr;
 	}
 
-	return SelectedItem;
+	return SelectedWorldBookmarkTreeItem;
 }
 
 UWorldBookmark* SWorldBookmarkBrowser::GetSelectedBookmark() const
@@ -1209,9 +993,9 @@ void SWorldBookmarkBrowser::SetSelectedBookmark(UWorldBookmark* InBookmark)
 
 	if (InBookmark)
 	{
-		FWorldBookmarkTreeItemPtr* BookmarkTreeItem = WorldBookmarkTreeItems.FindByPredicate([InBookmark](FWorldBookmarkTreeItemPtr InItem)
+		FTreeItemPtr* BookmarkTreeItem = TreeItems.FindByPredicate([InBookmark](FTreeItemPtr InItem)
 		{
-			return InItem->BookmarkAsset == InBookmark;
+			return InItem->IsA<FWorldBookmarkTreeItem>() && InItem->Cast<FWorldBookmarkTreeItem>()->BookmarkAsset == InBookmark;
 		});
 
 		if (BookmarkTreeItem)
@@ -1226,19 +1010,19 @@ void SWorldBookmarkBrowser::SetSelectedBookmark(UWorldBookmark* InBookmark)
 	}
 }
 
-void SWorldBookmarkBrowser::ItemScrolledIntoView(FWorldBookmarkTreeItemPtr WorldBookmarkTreeItem, const TSharedPtr<ITableRow>& Widget)
+void SWorldBookmarkBrowser::OnItemScrolledIntoView(FTreeItemPtr InTreeItem, const TSharedPtr<ITableRow>& InWidget)
 {
-	if (BookmarkPendingRename && WorldBookmarkTreeItem == BookmarkPendingRename)
+	if (TreeItemPendingRename && InTreeItem == TreeItemPendingRename)
 	{
 		// Abort the rename if our widget has lost focus.
 		TSharedPtr<SWindow> OwnerWindow = FSlateApplication::Get().FindWidgetWindow(AsShared());
 		if (OwnerWindow && OwnerWindow->HasAnyUserFocusOrFocusedDescendants())
 		{
 			// We know the tree item is now visible, bring up the editable text block.
-			BookmarkPendingRename->OnRenameRequested.ExecuteIfBound();
+			TreeItemPendingRename->OnRenameRequested.ExecuteIfBound();
 		}
 
-		BookmarkPendingRename.Reset();
+		TreeItemPendingRename.Reset();
 	}
 }
 
@@ -1247,7 +1031,7 @@ void SWorldBookmarkBrowser::RequestRenameSelectedBookmark()
 	if (FWorldBookmarkTreeItemPtr SelectedBookmarkTreeItem = GetSelectedBookmarkTreeItem())
 	{
 		// Before showing the inline editable text block, we must ensure that the tree item is visible.
-		BookmarkPendingRename = SelectedBookmarkTreeItem;
+		TreeItemPendingRename = SelectedBookmarkTreeItem;
 		BookmarksView->RequestScrollIntoView(SelectedBookmarkTreeItem);
 	}
 }
diff --git a/Engine/Source/Editor/WorldBookmark/Private/WorldBookmark/Browser/SWorldBookmarkBrowser.h b/Engine/Source/Editor/WorldBookmark/Private/WorldBookmark/Browser/SWorldBookmarkBrowser.h
--- a/Engine/Source/Editor/WorldBookmark/Private/WorldBookmark/Browser/SWorldBookmarkBrowser.h
+++ b/Engine/Source/Editor/WorldBookmark/Private/WorldBookmark/Browser/SWorldBookmarkBrowser.h
@@ -2,75 +2,18 @@
 #pragma once
 
 #include "AssetTextFilter.h"
-#include "ContentBrowserItem.h"
 #include "Misc/TextFilter.h"
 #include "Templates/SharedPointer.h"
 #include "UObject/ObjectSaveContext.h"
 #include "UnrealEdMisc.h"
-#include "Widgets/Views/SHeaderRow.h"	
 #include "Widgets/Views/STreeView.h"
 
+#include "WorldBookmark/Browser/BookmarkTreeItem.h"
+
 class UWorldBookmark;
 class IDetailsView;
 class SAssetView;
 
-namespace WorldBookmarkBrowserColumn
-{
-	struct FColumnDefinition
-	{
-		FName Id;
-		FText DisplayText;
-		FText ToolTipText;
-	};
-
-	extern const FColumnDefinition Favorite;
-	extern const FColumnDefinition RecentlyUsed;
-	extern const FColumnDefinition Label;
-	extern const FColumnDefinition World;
-
-	extern const TArray<FName> DefaultHiddenColumns;
-}
-
-struct FWorldBookmarkTreeItem : TSharedFromThis<FWorldBookmarkTreeItem>
-{
-public:
-	FAssetData BookmarkAsset;
-
-	const FSoftObjectPath& GetBookmarkWorld() const;
-
-	bool bInEditingMode = false;
-
-	/** Delegate for when the context menu requests a rename */
-	DECLARE_DELEGATE(FOnRenameRequested);
-	FOnRenameRequested OnRenameRequested;
-
-private:
-	mutable FSoftObjectPath CachedBookmarkWorld;
-};
-
-typedef TSharedPtr<FWorldBookmarkTreeItem> FWorldBookmarkTreeItemPtr;
-typedef TSharedRef<FWorldBookmarkTreeItem> FWorldBookmarkTreeItemRef;
-
-class FWorldBookmarkTableRow : public SMultiColumnTableRow<FWorldBookmarkTreeItemPtr>
-{
-public:
-	void Construct(const FArguments& InArgs, const TSharedRef<STableViewBase>& OwnerTable, FWorldBookmarkTreeItemRef InItem);
-
-private:
-	virtual TSharedRef<SWidget> GenerateWidgetForColumn(const FName& ColumnName) override;
-
-	FText GetLastAccessedText() const;
-	FText GetWorldNameText() const;
-
-	bool OnLabelTextVerifyChanged(const FText& InNewText, FText& OutErrorMessage);
-	void OnLabelTextCommitted(const FText& NewText, ETextCommit::Type InTextCommit);
-	void OnEnterEditingMode();
-	void OnExitEditingMode();
-
-private:
-	TWeakPtr<FWorldBookmarkTreeItem> WeakWorldBookmarkTreeItem;
-};
-
 class SWorldBookmarkBrowser : public SCompoundWidget
 {
 public:
@@ -84,16 +27,19 @@ public:
 	void Construct(const FArguments& InArgs);
 
 private:
-	void OnGetChildren(FWorldBookmarkTreeItemPtr InItem, TArray<FWorldBookmarkTreeItemPtr>& OutChildren);
-	TSharedRef<ITableRow> OnGenerateRow(FWorldBookmarkTreeItemPtr InTreeItem, const TSharedRef<STableViewBase>& OwnerTable);
-	void OnSelectedBookmarkChanged(FWorldBookmarkTreeItemPtr InTreeItem, ESelectInfo::Type SelectionType) const;
-	void OnBookmarkDoubleClicked(FWorldBookmarkTreeItemPtr InTreeItem);
+	// ~Begin STreeView delegates
+	void OnGetChildren(FTreeItemPtr InItem, TArray<FTreeItemPtr>& OutChildren);
+	TSharedRef<ITableRow> OnGenerateRow(FTreeItemPtr InTreeItem, const TSharedRef<STableViewBase>& OwnerTable);
+	void OnSelectionChanged(FTreeItemPtr InTreeItem, ESelectInfo::Type SelectionType);
+	void OnMouseButtonDoubleClicked(FTreeItemPtr InTreeItem);
+	TSharedPtr<SWidget> OnContextMenuOpening();
+	void OnItemScrolledIntoView(FTreeItemPtr WorldBookmarkTreeItem, const TSharedPtr<ITableRow>& Widget);
+	// ~End STreeView delegates
+
 	FReply OnCreateNewWorldBookmarkClicked();
 	void OnSearchBoxTextChanged(const FText& InSearchText);
 	void OnColumnSortChanged(const EColumnSortPriority::Type SortPriority, const FName& ColumnId, const EColumnSortMode::Type InSortMode);
-	TSharedPtr<SWidget> OnContextMenuOpening();
 	TSharedRef<SWidget> GetSettingsMenuContent();
-	void ItemScrolledIntoView(FWorldBookmarkTreeItemPtr WorldBookmarkTreeItem, const TSharedPtr<ITableRow>& Widget);
 
 	/** SWidget interface */
 	virtual void Tick(const FGeometry& AllottedGeometry, const double InCurrentTime, const float InDeltaTime) override;
@@ -120,6 +66,7 @@ private:
 
 	void OnUncontrolledChangelistModuleChanged();
 
+	FTreeItemPtr GetSelectedTreeItem() const;
 	FWorldBookmarkTreeItemPtr GetSelectedBookmarkTreeItem() const;
 	UWorldBookmark* GetSelectedBookmark() const;
 	void SetSelectedBookmark(UWorldBookmark* InBookmark);
@@ -154,9 +101,9 @@ private:
 	FDelegateHandle OnSettingsChangedHandle;
 	FDelegateHandle OnUncontrolledChangelistModuleChangedHandle;
 
-	TArray<FWorldBookmarkTreeItemPtr>	WorldBookmarkTreeItems;
+	TArray<FTreeItemPtr> TreeItems;
 
-	TSharedPtr<STreeView<FWorldBookmarkTreeItemPtr>> BookmarksView;
+	TSharedPtr<STreeView<FTreeItemPtr>> BookmarksView;
 
 	TSharedPtr<IDetailsView> PropertyView;
 
@@ -166,7 +113,7 @@ private:
 	
 	FText CurrentSearchString;
 	
-	FWorldBookmarkTreeItemPtr BookmarkPendingRename;
+	FTreeItemPtr TreeItemPendingRename;
 
 	bool bIsInTick;
 	bool bPendingRefresh;
diff --git a/Engine/Source/Editor/WorldBookmark/Private/WorldBookmark/Browser/Settings.cpp b/Engine/Source/Editor/WorldBookmark/Private/WorldBookmark/Browser/Settings.cpp
new file mode 100644
--- /dev/null
+++ b/Engine/Source/Editor/WorldBookmark/Private/WorldBookmark/Browser/Settings.cpp
@@ -0,0 +1,112 @@
+// Copyright Epic Games, Inc. All Rights Reserved.
+
+#include "WorldBookmark/Browser/Settings.h"
+
+#include "WorldBookmark/Browser/Columns.h"
+
+TObjectPtr<UWorldBookmarkBrowserSettings> UWorldBookmarkBrowserSettings::Instance = nullptr;
+
+UWorldBookmarkBrowserSettings& UWorldBookmarkBrowserSettings::Get()
+{
+	if (!Instance)
+	{
+		Instance = NewObject<UWorldBookmarkBrowserSettings>();
+
+		Instance->HiddenColumns = WorldBookmarkBrowserColumn::DefaultHiddenColumns;
+
+		Instance->LoadEditorConfig();
+		Instance->AddToRoot();
+	}
+
+	return *Instance;
+}
+
+void UWorldBookmarkBrowserSettings::ApplySettingsChanges()
+{
+	SaveEditorConfig();
+	OnSettingsChangedDelegate.Broadcast();
+}
+
+UWorldBookmarkBrowserSettings::FWorldBookmarkBrowserSettingChanged& UWorldBookmarkBrowserSettings::OnSettingsChanged()
+{
+	UWorldBookmarkBrowserSettings& Settings = Get();
+	return Settings.OnSettingsChangedDelegate;
+}
+
+bool UWorldBookmarkBrowserSettings::GetShowOnlyBookmarksForCurrentWorld()
+{
+	return Get().bShowOnlyBookmarksForCurrentWorld;
+}
+
+void UWorldBookmarkBrowserSettings::ToggleShowOnlyBookmarksForCurrentWorld()
+{
+	UWorldBookmarkBrowserSettings& Settings = Get();
+	Settings.bShowOnlyBookmarksForCurrentWorld = !Settings.bShowOnlyBookmarksForCurrentWorld;
+	Settings.ApplySettingsChanges();
+}
+
+bool UWorldBookmarkBrowserSettings::GetShowOnlyUncontrolledBookmarks()
+{
+	return Get().bShowOnlyUncontrolledBookmarks;
+}
+
+void UWorldBookmarkBrowserSettings::ToggleShowOnlyUncontrolledBookmarks()
+{
+	UWorldBookmarkBrowserSettings& Settings = Get();
+	Settings.bShowOnlyUncontrolledBookmarks = !Settings.bShowOnlyUncontrolledBookmarks;
+	Settings.ApplySettingsChanges();
+}
+
+bool UWorldBookmarkBrowserSettings::GetShowOnlyFavoriteBookmarks()
+{
+	return Get().bShowOnlyFavoriteBookmarks;
+}
+
+void UWorldBookmarkBrowserSettings::ToggleShowOnlyFavoriteBookmarks()
+{
+	UWorldBookmarkBrowserSettings& Settings = Get();
+	Settings.bShowOnlyFavoriteBookmarks = !Settings.bShowOnlyFavoriteBookmarks;
+	Settings.ApplySettingsChanges();
+}
+
+bool UWorldBookmarkBrowserSettings::GetShowOnlyLastRecentlyUsedBookmarks()
+{
+	return Get().bShowOnlyLastRecentlyUsedBookmarks;
+}
+
+void UWorldBookmarkBrowserSettings::ToggleShowOnlyLastRecentlyUsedBookmarks()
+{
+	UWorldBookmarkBrowserSettings& Settings = Get();
+	Settings.bShowOnlyLastRecentlyUsedBookmarks = !Settings.bShowOnlyLastRecentlyUsedBookmarks;
+	Settings.ApplySettingsChanges();
+}
+
+TArray<FName> UWorldBookmarkBrowserSettings::GetHiddenColumns()
+{
+	return Get().HiddenColumns;
+}
+
+void UWorldBookmarkBrowserSettings::SetHiddenColumns(const TArray<FName>& InHiddenColumns)
+{
+	UWorldBookmarkBrowserSettings& Settings = Get();
+	if (Settings.HiddenColumns != InHiddenColumns)
+	{
+		Settings.HiddenColumns = InHiddenColumns;
+		Settings.ApplySettingsChanges();
+	}
+}
+
+int32 UWorldBookmarkBrowserSettings::GetMaxLastRecentlyUsedItems()
+{
+	return Get().MaxLastRecentlyUsedItems;
+}
+
+void UWorldBookmarkBrowserSettings::SetMaxLastRecentlyUsedItems(int32 InMaxLastRecentlyUsedItems)
+{
+	UWorldBookmarkBrowserSettings& Settings = Get();
+	if (Settings.MaxLastRecentlyUsedItems != InMaxLastRecentlyUsedItems)
+	{
+		Settings.MaxLastRecentlyUsedItems = InMaxLastRecentlyUsedItems;
+		Settings.ApplySettingsChanges();
+	}
+}
\ No newline at end of file
diff --git a/Engine/Source/Editor/WorldBookmark/Private/WorldBookmark/Browser/Settings.h b/Engine/Source/Editor/WorldBookmark/Private/WorldBookmark/Browser/Settings.h
new file mode 100644
--- /dev/null
+++ b/Engine/Source/Editor/WorldBookmark/Private/WorldBookmark/Browser/Settings.h
@@ -0,0 +1,62 @@
+// Copyright Epic Games, Inc. All Rights Reserved.
+#pragma once
+
+#include "EditorConfigBase.h"
+
+#include "Settings.generated.h"
+
+UCLASS(EditorConfig = "WorldBookmarkBrowser")
+class UWorldBookmarkBrowserSettings : public UEditorConfigBase
+{
+	GENERATED_BODY()
+
+public:
+	DECLARE_MULTICAST_DELEGATE(FWorldBookmarkBrowserSettingChanged)
+
+	static FWorldBookmarkBrowserSettingChanged& OnSettingsChanged();
+
+	static bool GetShowOnlyBookmarksForCurrentWorld();
+	static void ToggleShowOnlyBookmarksForCurrentWorld();
+
+	static bool GetShowOnlyUncontrolledBookmarks();
+	static void ToggleShowOnlyUncontrolledBookmarks();
+
+	static bool GetShowOnlyFavoriteBookmarks();
+	static void ToggleShowOnlyFavoriteBookmarks();
+
+	static bool GetShowOnlyLastRecentlyUsedBookmarks();
+	static void ToggleShowOnlyLastRecentlyUsedBookmarks();
+
+	static TArray<FName> GetHiddenColumns();
+	static void SetHiddenColumns(const TArray<FName>& InHiddenColumns);
+
+	static int32 GetMaxLastRecentlyUsedItems();
+	static void SetMaxLastRecentlyUsedItems(int32 InMaxLastRecentlyUsedItems);
+
+private:
+	static UWorldBookmarkBrowserSettings& Get();
+
+	void ApplySettingsChanges();
+
+	UPROPERTY(meta=(EditorConfig))
+	bool bShowOnlyBookmarksForCurrentWorld = true;
+
+	UPROPERTY(meta=(EditorConfig))
+	bool bShowOnlyUncontrolledBookmarks = false;
+
+	UPROPERTY(meta=(EditorConfig))
+	bool bShowOnlyFavoriteBookmarks = false;
+
+	UPROPERTY(meta=(EditorConfig))
+	bool bShowOnlyLastRecentlyUsedBookmarks = false;
+
+	UPROPERTY(meta=(EditorConfig))
+	int32 MaxLastRecentlyUsedItems = 5;
+
+	UPROPERTY(meta=(EditorConfig))
+	TArray<FName> HiddenColumns;
+
+	FWorldBookmarkBrowserSettingChanged OnSettingsChangedDelegate;
+
+	static TObjectPtr<UWorldBookmarkBrowserSettings> Instance;
+};
diff --git a/Engine/Source/Editor/WorldBookmark/Private/WorldBookmark/Browser/TreeItem.cpp b/Engine/Source/Editor/WorldBookmark/Private/WorldBookmark/Browser/TreeItem.cpp
new file mode 100644
--- /dev/null
+++ b/Engine/Source/Editor/WorldBookmark/Private/WorldBookmark/Browser/TreeItem.cpp
@@ -0,0 +1,3 @@
+// Copyright Epic Games, Inc. All Rights Reserved.
+
+#include "WorldBookmark/Browser/TreeItem.h"
diff --git a/Engine/Source/Editor/WorldBookmark/Private/WorldBookmark/Browser/TreeItem.h b/Engine/Source/Editor/WorldBookmark/Private/WorldBookmark/Browser/TreeItem.h
new file mode 100644
--- /dev/null
+++ b/Engine/Source/Editor/WorldBookmark/Private/WorldBookmark/Browser/TreeItem.h
@@ -0,0 +1,58 @@
+// Copyright Epic Games, Inc. All Rights Reserved.
+#pragma once
+
+#include "Delegates/Delegate.h"
+#include "Templates/SharedPointer.h"
+
+enum ETreeItemType
+{
+	Folder,
+	WorldBookmark
+};
+
+// Define a common base class for your items
+struct ITreeItem : public TSharedFromThis<ITreeItem>
+{
+public:
+	virtual FString GetName() const = 0;
+
+	/** Attempt to cast this item to another type if it is of that type. Returns nullptr if it fails */
+	template <typename T>
+	T* Cast()
+	{
+		return TreeItemType == T::Type ? StaticCast<T*>(this) : nullptr;
+	}
+
+	/** Attempt to cast this item to another type if it is of that type. Returns nullptr if it fails */
+	template <typename T>
+	const T* Cast() const
+	{
+		return TreeItemType == T::Type ? StaticCast<const T*>(this) : nullptr;
+	}
+
+	/** Returns true if this item is of the specified type */
+	template <typename T>
+	bool IsA() const
+	{
+		return TreeItemType == T::Type;
+	}
+	
+	bool bInEditingMode = false;
+
+	/** Delegate for when the context menu requests a rename */
+	DECLARE_DELEGATE(FOnRenameRequested);
+	FOnRenameRequested OnRenameRequested;
+
+	/** Tree item type identifier */
+	ETreeItemType TreeItemType;
+
+protected:
+	/** Default constructor */
+	ITreeItem(ETreeItemType InType) : TreeItemType(InType) {}
+	virtual ~ITreeItem()
+	{
+	}
+};
+
+typedef TSharedPtr<ITreeItem> FTreeItemPtr;
+typedef TSharedRef<ITreeItem> FTreeItemRef;
diff --git a/Engine/Source/Editor/WorldBookmark/Private/WorldBookmark/Browser/WorldBookmarkBrowserSettings.cpp b/Engine/Source/Editor/WorldBookmark/Private/WorldBookmark/Browser/WorldBookmarkBrowserSettings.cpp
deleted file mode 100644
diff --git a/Engine/Source/Editor/WorldBookmark/Private/WorldBookmark/Browser/WorldBookmarkBrowserSettings.h b/Engine/Source/Editor/WorldBookmark/Private/WorldBookmark/Browser/WorldBookmarkBrowserSettings.h
deleted file mode 100644
