World Bookmark Browser - Display hints based on selection
* Hints replace the property view which is shown when the selected bookmark is valid
* Show an hint if no bookmark is selected
* Show an hint if the selected bookmark has no state (currently only possible if the new bookmark is created in an unsaved world)

diff --git a/Engine/Source/Editor/UnrealEd/Public/EditorState/EditorStateCollection.h b/Engine/Source/Editor/UnrealEd/Public/EditorState/EditorStateCollection.h
--- a/Engine/Source/Editor/UnrealEd/Public/EditorState/EditorStateCollection.h
+++ b/Engine/Source/Editor/UnrealEd/Public/EditorState/EditorStateCollection.h
@@ -46,6 +46,11 @@ struct FEditorStateCollection
 		ForEachState(this, InFunc, InEditorStateTypeFilter);
 	}
 
+	bool HasStates() const
+	{
+		return !EditorStates.IsEmpty();
+	}
+
 	TArray<UEditorState*> GetStates() const
 	{
 		TArray<UEditorState*> Values;
diff --git a/Engine/Source/Editor/WorldBookmark/Private/WorldBookmark/Browser/SWorldBookmarkBrowser.cpp b/Engine/Source/Editor/WorldBookmark/Private/WorldBookmark/Browser/SWorldBookmarkBrowser.cpp
--- a/Engine/Source/Editor/WorldBookmark/Private/WorldBookmark/Browser/SWorldBookmarkBrowser.cpp
+++ b/Engine/Source/Editor/WorldBookmark/Private/WorldBookmark/Browser/SWorldBookmarkBrowser.cpp
@@ -232,16 +232,10 @@ void SWorldBookmarkBrowser::Construct(const FArguments& InArgs)
 			]
 		]
 
-		// Details Propery View
+		// Details Property View
 		+ SSplitter::Slot()
 		[
-			SNew(SVerticalBox)
-
-			+ SVerticalBox::Slot()
-			.Padding(2, 4, 0, 0)
-			[
-				PropertyView.ToSharedRef()
-			]
+			SAssignNew(DetailsBox, SVerticalBox)
 		]
 	];
 
@@ -286,6 +280,8 @@ void SWorldBookmarkBrowser::Construct(const FArguments& InArgs)
 	Commands->MapAction(FWorldBookmarkCommands::Get().MoveCameraToLocation,
 		FExecuteAction::CreateSP(this, &SWorldBookmarkBrowser::GoToSelectedBookmarkLocation),
 		FCanExecuteAction::CreateSP(this, &SWorldBookmarkBrowser::CanGoToSelectedBookmarkLocation));
+
+	UpdateDetailsView(/*InSelectedWorldBookmark=*/nullptr);
 }
 
 TSharedRef<SWidget> SWorldBookmarkBrowser::GetSettingsMenuContent()
@@ -382,10 +378,10 @@ FReply SWorldBookmarkBrowser::OnKeyDown(const FGeometry& MyGeometry, const FKeyE
 	return FReply::Unhandled();
 }
 
-void SWorldBookmarkBrowser::RefreshItems()
+void SWorldBookmarkBrowser::RefreshItems(bool bForceRefresh)
 {
 	// Never perform a refresh outside of the tick method
-	if (!bIsInTick)
+	if (!bIsInTick && !bForceRefresh)
 	{
 		bPendingRefresh = true;
 		return;
@@ -636,6 +632,7 @@ void SWorldBookmarkBrowser::OnSelectionChanged(FTreeItemPtr InTreeItem, ESelectI
 {
 	UWorldBookmark* SelectedWorldBookmark = GetSelectedBookmark();
 	PropertyView->SetObject(SelectedWorldBookmark);
+	UpdateDetailsView(SelectedWorldBookmark);
 }
 
 void SWorldBookmarkBrowser::OnMouseButtonDoubleClicked(FTreeItemPtr InTreeItem)
@@ -661,8 +658,8 @@ FReply SWorldBookmarkBrowser::OnCreateNewWorldBookmarkClicked()
 		// Show it in the content browser
 		GEditor->SyncBrowserToObject(NewWorldBookmark);
 
-		// Refresh the bookmark browser to ensure the new bookmark is shown
-		RefreshItems();
+		// Force refresh the bookmark browser, to ensure the new bookmark tree item is created
+		RefreshItems(/*bForceRefresh=*/true);
 
 		// Make sure the new bookmark is selected
 		SetSelectedBookmark(NewWorldBookmark);
@@ -1132,5 +1129,45 @@ void SWorldBookmarkBrowser::Tick(const FGeometry& AllottedGeometry, const double
 	SCompoundWidget::Tick(AllottedGeometry, InCurrentTime, InDeltaTime);
 }
 
+void SWorldBookmarkBrowser::UpdateDetailsView(UWorldBookmark* SelectedWorldBookmark)
+{
+	DetailsBox->ClearChildren();
+
+	// Show message if no bookmark is selected
+	if (SelectedWorldBookmark == nullptr)
+	{
+		DetailsBox->AddSlot()
+			.HAlign(HAlign_Center)
+			.VAlign(VAlign_Top)
+			.Padding(2.0f, 24.0f, 2.0f, 2.0f)
+			[
+				SNew(STextBlock)
+					.Text(LOCTEXT("NoBookmarkSelected", "Select a world bookmark to view details."))
+					.TextStyle(FAppStyle::Get(), "HintText")
+			];
+	}
+	// Show message if bookmark has no data
+	else if (!SelectedWorldBookmark->HasEditorStates())
+	{
+		DetailsBox->AddSlot()
+			.HAlign(HAlign_Center)
+			.Padding(2.0f, 24.0f, 2.0f, 2.0f)
+			[
+				SNew(STextBlock)
+					.Text(LOCTEXT("UninitializedBookmarkSelected", "Bookmark is empty. Update it to view details."))
+					.TextStyle(FAppStyle::Get(), "HintText")
+			];
+	}
+	// Show details property
+	else
+	{
+		DetailsBox->AddSlot()
+			.Padding(2, 4, 0, 0)
+			[
+				PropertyView.ToSharedRef()
+			];
+	}
+}
+
 #undef LOCTEXT_NAMESPACE
 
diff --git a/Engine/Source/Editor/WorldBookmark/Private/WorldBookmark/Browser/SWorldBookmarkBrowser.h b/Engine/Source/Editor/WorldBookmark/Private/WorldBookmark/Browser/SWorldBookmarkBrowser.h
--- a/Engine/Source/Editor/WorldBookmark/Private/WorldBookmark/Browser/SWorldBookmarkBrowser.h
+++ b/Engine/Source/Editor/WorldBookmark/Private/WorldBookmark/Browser/SWorldBookmarkBrowser.h
@@ -45,7 +45,7 @@ private:
 	virtual void Tick(const FGeometry& AllottedGeometry, const double InCurrentTime, const float InDeltaTime) override;
 	virtual FReply OnKeyDown(const FGeometry& MyGeometry, const FKeyEvent& InKeyEvent) override;
 
-	void RefreshItems();
+	void RefreshItems(bool bForceRefresh = false);
 	void SortItems();
 
 	EColumnSortMode::Type GetColumnSortMode(const FName ColumnId) const;
@@ -90,6 +90,8 @@ private:
 	void PlayFromSelectedBookmarkLocation();
 	void GoToSelectedBookmarkLocation();
 
+	void UpdateDetailsView(UWorldBookmark* InSelectedWorldBookmark);
+
 private:
 	FDelegateHandle OnAssetAddedHandle;
 	FDelegateHandle OnAssetRemovedHandle;
@@ -104,8 +106,8 @@ private:
 	TArray<FTreeItemPtr> TreeItems;
 
 	TSharedPtr<STreeView<FTreeItemPtr>> BookmarksView;
-
 	TSharedPtr<IDetailsView> PropertyView;
+	TSharedPtr<SVerticalBox> DetailsBox;
 
 	TSharedPtr<FUICommandList> Commands;
 
diff --git a/Engine/Source/Editor/WorldBookmark/Private/WorldBookmark/WorldBookmark.cpp b/Engine/Source/Editor/WorldBookmark/Private/WorldBookmark/WorldBookmark.cpp
--- a/Engine/Source/Editor/WorldBookmark/Private/WorldBookmark/WorldBookmark.cpp
+++ b/Engine/Source/Editor/WorldBookmark/Private/WorldBookmark/WorldBookmark.cpp
@@ -127,6 +127,16 @@ bool UWorldBookmark::CanLoad(FText* FailureReason) const
 		return false;
 	}
 
+	if (!HasEditorStates())
+	{
+		if (FailureReason)
+		{
+			*FailureReason = LOCTEXT("CanLoadFailure_IsUninitialized", "Can't load empty bookmark");
+		}
+
+		return false;
+	}
+
 	return true;
 }
 
@@ -243,6 +253,11 @@ void UWorldBookmark::UpdateStates(const TArray<TSubclassOf<UEditorState>>& InSta
 	}
 }
 
+bool UWorldBookmark::HasEditorStates() const
+{
+	return EditorState.HasStates();
+}
+
 bool UWorldBookmark::GetIsUserFavorite() const
 {
 	return bFavorite;
diff --git a/Engine/Source/Editor/WorldBookmark/Public/WorldBookmark/WorldBookmark.h b/Engine/Source/Editor/WorldBookmark/Public/WorldBookmark/WorldBookmark.h
--- a/Engine/Source/Editor/WorldBookmark/Public/WorldBookmark/WorldBookmark.h
+++ b/Engine/Source/Editor/WorldBookmark/Public/WorldBookmark/WorldBookmark.h
@@ -33,6 +33,9 @@ class UWorldBookmark : public UObject
 	/** Update a set of states for the bookmark. */
 	void UpdateStates(const TArray<TSubclassOf<UEditorState>>& InStatesToUpdate);
 
+	/** Return true if the world bookmark contains states that can be restored. */
+	bool HasEditorStates() const;
+
 	/** Return true if this bookmark was flagged as being a favorite bookmark of the user. */
 	bool GetIsUserFavorite() const;
 
