World Bookmarks - Added folders & categories to the world bookmark browser
* Introduced a tree view mode to the browser
* Folders are based on the actual bookmark asset path.
* Categories are user defined (name + color). They are stored in the bookmark metadata, so the full list of categories is always dynamically populated (no predefined list/config is needed)

diff --git a/Engine/Source/Editor/WorldBookmark/Private/WorldBookmark/Browser/BookmarkTableRow.cpp b/Engine/Source/Editor/WorldBookmark/Private/WorldBookmark/Browser/BookmarkTableRow.cpp
--- a/Engine/Source/Editor/WorldBookmark/Private/WorldBookmark/Browser/BookmarkTableRow.cpp
+++ b/Engine/Source/Editor/WorldBookmark/Private/WorldBookmark/Browser/BookmarkTableRow.cpp
@@ -1,13 +1,16 @@
 // Copyright Epic Games, Inc. All Rights Reserved.
 
 #include "WorldBookmark/Browser/BookmarkTableRow.h"
+#include "WorldBookmark/Browser/Columns.h"
+#include "WorldBookmark/Browser/CommonTableRow.h"
+#include "WorldBookmark/Browser/Icons.h"
+#include "WorldBookmark/Browser/Settings.h"
+#include "WorldBookmark/WorldBookmark.h"
 
 #include "AssetToolsModule.h"
 #include "AssetViewUtils.h"
 #include "Modules/ModuleManager.h"
-#include "WorldBookmark/Browser/Columns.h"
-#include "WorldBookmark/Browser/Icons.h"
-#include "WorldBookmark/WorldBookmark.h"
+#include "Widgets/Colors/SColorBlock.h"
 #include "Widgets/Text/SInlineEditableTextBlock.h"
 #include "Widgets/Text/STextBlock.h"
 #include "Widgets/Views/SListView.h"
@@ -16,13 +19,16 @@
 
 void FWorldBookmarkTableRow::Construct(const FArguments& InArgs, const TSharedRef<STableViewBase>& OwnerTable, FWorldBookmarkTreeItemRef InItem)
 {
-	WeakWorldBookmarkTreeItem = InItem;	
-	SMultiColumnTableRow<FWorldBookmarkTreeItemPtr>::Construct(InArgs, OwnerTable);
+	WeakTreeItem = InItem;
+	SMultiColumnTableRow<FWorldBookmarkTreeItemPtr>::Construct(
+		STableRow::FArguments(InArgs)
+		.OnDragDetected(this, &FWorldBookmarkTableRowBase::OnRowDragDetected),
+		OwnerTable);
 }
 
 TSharedRef<SWidget> FWorldBookmarkTableRow::GenerateWidgetForColumn(const FName& ColumnName)
 {
-	if (FWorldBookmarkTreeItemPtr WorldBookmarkTreeItem = WeakWorldBookmarkTreeItem.Pin())
+	if (FWorldBookmarkTreeItemPtr WorldBookmarkTreeItem = GetTreeItem())
 	{
 		if (ColumnName == WorldBookmarkBrowserColumn::Favorite.Id)
 		{
@@ -47,39 +53,18 @@ TSharedRef<SWidget> FWorldBookmarkTableRow::GenerateWidgetForColumn(const FName&
 		}
 		else if (ColumnName == WorldBookmarkBrowserColumn::Label.Id)
 		{
-			TSharedPtr<SInlineEditableTextBlock> InlineEditableTextBlock;
-
-			TSharedRef<SWidget> LabelWidget =
-				SNew(SVerticalBox)
-				+ SVerticalBox::Slot()
-				.Padding(10, 5, 0, 0)
-				.AutoHeight()
-				[
-					SAssignNew(InlineEditableTextBlock, SInlineEditableTextBlock)
-						.Text(FText::FromName(WorldBookmarkTreeItem->BookmarkAsset.AssetName))
-						.OnVerifyTextChanged(this, &FWorldBookmarkTableRow::OnLabelTextVerifyChanged)
-						.OnTextCommitted(this, &FWorldBookmarkTableRow::OnLabelTextCommitted)
-						.OnEnterEditingMode(this, &FWorldBookmarkTableRow::OnEnterEditingMode)
-						.OnExitEditingMode(this, &FWorldBookmarkTableRow::OnExitEditingMode)
-						.IsSelected(this, &FWorldBookmarkTableRow::IsSelected)
-				]
-				+ SVerticalBox::Slot()
-				.Padding(10, 0, 0, 5)
-				.AutoHeight()
-				[
-					SNew(STextBlock)
-						.ColorAndOpacity(FSlateColor::UseSubduedForeground())
-						.Text(this, &FWorldBookmarkTableRow::GetLastAccessedText)
-				];
-
-			WorldBookmarkTreeItem->OnRenameRequested.BindSP(InlineEditableTextBlock.Get(), &SInlineEditableTextBlock::EnterEditingMode);
-
-			return LabelWidget;
+			if (UWorldBookmarkBrowserSettings::IsViewMode(EWorldBookmarkBrowserViewMode::TreeView))
+			{
+				return GenerateLabelForTreeView();
+			}
+			else
+			{
+				return GenerateLabelForListView();
+			}
 		}
 		else if (ColumnName == WorldBookmarkBrowserColumn::World.Id)
 		{
 			return SNew(SBox)
-				.Padding(2.0f)
 				.HAlign(HAlign_Left)
 				.VAlign(VAlign_Center)
 				[
@@ -87,99 +72,113 @@ TSharedRef<SWidget> FWorldBookmarkTableRow::GenerateWidgetForColumn(const FName&
 						.Text(this, &FWorldBookmarkTableRow::GetWorldNameText)
 				];
 		}
+		else if (ColumnName == WorldBookmarkBrowserColumn::Category.Id)
+		{
+			return SNew(SBox)
+				.HAlign(HAlign_Center)
+				.VAlign(VAlign_Center)
+				[
+					SNew(SColorBlock)
+						.Color(this, &FWorldBookmarkTableRow::GetCategoryColor)
+						.ToolTipText(this, &FWorldBookmarkTableRow::GetCategoryText)
+						.CornerRadius(FVector4(4.0f, 4.0f, 4.0f, 4.0f))
+				];
+		}
 	}
 
 	return SNullWidget::NullWidget;
 }
 
-static FSoftObjectPath GetBookmarkObjectPath(FWorldBookmarkTreeItemPtr InWorldBookmarkTreeItem)
-{
-	return InWorldBookmarkTreeItem->BookmarkAsset.GetSoftObjectPath();
-}
-
-static FSoftObjectPath GetRenamedBookmarkObjectPath(FWorldBookmarkTreeItemPtr InWorldBookmarkTreeItem, const FText& InNewAssetName)
+TSharedRef<SWidget> FWorldBookmarkTableRow::GenerateLabelForTreeView()
 {
-	const FString PackagePath = InWorldBookmarkTreeItem->BookmarkAsset.PackagePath.ToString();
-	const FString NewObjectPath = FString::Format(TEXT("{0}/{1}.{1}"), { *PackagePath, *InNewAssetName.ToString() });
-	return FSoftObjectPath(NewObjectPath);
+	return CreateTreeLabelWidget(SharedThis(this));
 }
 
-bool FWorldBookmarkTableRow::OnLabelTextVerifyChanged(const FText& InNewText, FText& OutErrorMessage)
+TSharedRef<SWidget> FWorldBookmarkTableRow::GenerateLabelForListView()
 {
-	if (FWorldBookmarkTreeItemPtr WorldBookmarkTreeItem = WeakWorldBookmarkTreeItem.Pin())
-	{
-		const FSoftObjectPath NewObjectPath = GetRenamedBookmarkObjectPath(WorldBookmarkTreeItem, InNewText);
-		return AssetViewUtils::IsValidObjectPathForCreate(NewObjectPath.ToString(), OutErrorMessage);
-	}
-
-	return false;
+	return SNew(SVerticalBox)
+		+ SVerticalBox::Slot()
+		.Padding(10, 5, 0, 0)
+		.AutoHeight()
+		[
+			CreateEditableLabelWidget(SharedThis(this))
+		]
+		+ SVerticalBox::Slot()
+		.Padding(10, 0, 0, 5)
+		.AutoHeight()
+		[
+			SNew(STextBlock)
+				.ColorAndOpacity(FSlateColor::UseSubduedForeground())
+				.Text(this, &FWorldBookmarkTableRow::GetLastAccessedText)
+		];
 }
 
-void FWorldBookmarkTableRow::OnLabelTextCommitted(const FText& InNewText, ETextCommit::Type InTextCommit)
+FText FWorldBookmarkTableRow::GetLastAccessedText() const
 {
-	FText ErrorMessage;
-	if (OnLabelTextVerifyChanged(InNewText, ErrorMessage))
+	if (FWorldBookmarkTreeItemPtr WorldBookmarkTreeItem = GetTreeItem())
 	{
-		if (FWorldBookmarkTreeItemPtr WorldBookmarkTreeItem = WeakWorldBookmarkTreeItem.Pin())
+		if (UWorldBookmark* WorldBookmark = Cast<UWorldBookmark>(WorldBookmarkTreeItem->BookmarkAsset.GetAsset()))
 		{
-			FAssetToolsModule& AssetToolsModule = FModuleManager::LoadModuleChecked<FAssetToolsModule>("AssetTools");
+			FDateTime LastLoadedUTC = WorldBookmark->GetUserLastLoadedTimeStampUTC();
 
-			TArray<FAssetRenameData> AssetsAndNames;
-			const FSoftObjectPath OldObjectPath = GetBookmarkObjectPath(WorldBookmarkTreeItem);
-			const FSoftObjectPath NewObjectPath = GetRenamedBookmarkObjectPath(WorldBookmarkTreeItem, InNewText);
-			new (AssetsAndNames) FAssetRenameData(OldObjectPath, NewObjectPath);
+			if (LastLoadedUTC != FDateTime::MinValue())
+			{
+				FTimespan UTCOffset = FDateTime::UtcNow() - FDateTime::Now();
+				uint64 AdjustedTicks = LastLoadedUTC.GetTicks() - UTCOffset.GetTicks();
+				FDateTime LastLoadedLocal = FDateTime(AdjustedTicks);
 
-			AssetToolsModule.Get().RenameAssetsWithDialog(AssetsAndNames);
+				return FText::Format(LOCTEXT("LastAccessedDateTime", "Last accessed {0}"), FText::AsDateTime(LastLoadedUTC, TEXT("%Y-%m-%d %H:%M")));
+			}
 		}
 	}
-}
 
-void FWorldBookmarkTableRow::OnEnterEditingMode()
-{
-	if (FWorldBookmarkTreeItemPtr WorldBookmarkTreeItem = WeakWorldBookmarkTreeItem.Pin())
-	{
-		WorldBookmarkTreeItem->bInEditingMode = true;
-	}
+	return FText();
 }
 
-void FWorldBookmarkTableRow::OnExitEditingMode()
+FText FWorldBookmarkTableRow::GetWorldNameText() const
 {
-	if (FWorldBookmarkTreeItemPtr WorldBookmarkTreeItem = WeakWorldBookmarkTreeItem.Pin())
+	if (FWorldBookmarkTreeItemPtr WorldBookmarkTreeItem = GetTreeItem())
 	{
-		WorldBookmarkTreeItem->bInEditingMode = false;
+		return FText::FromString(WorldBookmarkTreeItem->GetBookmarkWorld().GetAssetName());
 	}
+
+	return FText();
 }
 
-FText FWorldBookmarkTableRow::GetLastAccessedText() const
+FLinearColor FWorldBookmarkTableRow::GetCategoryColor() const
 {
-	if (FWorldBookmarkTreeItemPtr WorldBookmarkTreeItem = WeakWorldBookmarkTreeItem.Pin())
+	if (FWorldBookmarkTreeItemPtr WorldBookmarkTreeItem = GetTreeItem())
 	{
 		if (UWorldBookmark* WorldBookmark = Cast<UWorldBookmark>(WorldBookmarkTreeItem->BookmarkAsset.GetAsset()))
 		{
-			FDateTime LastLoadedUTC = WorldBookmark->GetUserLastLoadedTimeStampUTC();
-
-			if (LastLoadedUTC != FDateTime::MinValue())
+			if (WorldBookmark->GetBookmarkCategory().Name != NAME_None)
 			{
-				FTimespan UTCOffset = FDateTime::UtcNow() - FDateTime::Now();
-				uint64 AdjustedTicks = LastLoadedUTC.GetTicks() - UTCOffset.GetTicks();
-				FDateTime LastLoadedLocal = FDateTime(AdjustedTicks);
-
-				return FText::Format(LOCTEXT("LastAccessedDateTime", "Last accessed {0}"), FText::AsDateTime(LastLoadedUTC, TEXT("%Y-%m-%d %H:%M")));
+				return WorldBookmark->GetBookmarkCategory().Color;
 			}
 		}
 	}
 
-	return FText();
+	return FLinearColor::Transparent;
 }
 
-FText FWorldBookmarkTableRow::GetWorldNameText() const
+FText FWorldBookmarkTableRow::GetCategoryText() const
 {
-	if (FWorldBookmarkTreeItemPtr WorldBookmarkTreeItem = WeakWorldBookmarkTreeItem.Pin())
+	if (FWorldBookmarkTreeItemPtr WorldBookmarkTreeItem = GetTreeItem())
 	{
-		return FText::FromString(WorldBookmarkTreeItem->GetBookmarkWorld().GetAssetName());
+		if (UWorldBookmark* WorldBookmark = Cast<UWorldBookmark>(WorldBookmarkTreeItem->BookmarkAsset.GetAsset()))
+		{
+			return FText::FromName(WorldBookmark->GetBookmarkCategory().Name);
+		}
 	}
 
 	return FText();
 }
 
+FSlateColor FWorldBookmarkTableRow::GetIconColor() const
+{
+	FLinearColor IconColor = GetCategoryColor();
+	
+	return IconColor == FLinearColor::Transparent ? FSlateColor::UseForeground() : IconColor;
+}
+
 #undef LOCTEXT_NAMESPACE
diff --git a/Engine/Source/Editor/WorldBookmark/Private/WorldBookmark/Browser/BookmarkTableRow.h b/Engine/Source/Editor/WorldBookmark/Private/WorldBookmark/Browser/BookmarkTableRow.h
--- a/Engine/Source/Editor/WorldBookmark/Private/WorldBookmark/Browser/BookmarkTableRow.h
+++ b/Engine/Source/Editor/WorldBookmark/Private/WorldBookmark/Browser/BookmarkTableRow.h
@@ -3,26 +3,25 @@
 
 #include "Widgets/Views/STableRow.h"
 #include "Widgets/Views/SListView.h"
-#include "WorldBookmark/Browser/BookmarkTreeItem.h"
 
-class SWidget;
+#include "WorldBookmark/Browser/BookmarkTreeItem.h"
+#include "WorldBookmark/Browser/CommonTableRow.h"
 
-class FWorldBookmarkTableRow : public SMultiColumnTableRow<FWorldBookmarkTreeItemPtr>
+class FWorldBookmarkTableRow : public SMultiColumnTableRow<FWorldBookmarkTreeItemPtr>, public FWorldBookmarkTableRowBase<FWorldBookmarkTreeItem>
 {
 public:
 	void Construct(const FArguments& InArgs, const TSharedRef<STableViewBase>& OwnerTable, FWorldBookmarkTreeItemRef InItem);
 
-private:
 	virtual TSharedRef<SWidget> GenerateWidgetForColumn(const FName& ColumnName) override;
 
+	TSharedRef<SWidget> GenerateLabelForTreeView();
+	TSharedRef<SWidget> GenerateLabelForListView();
+
 	FText GetLastAccessedText() const;
 	FText GetWorldNameText() const;
 
-	bool OnLabelTextVerifyChanged(const FText& InNewText, FText& OutErrorMessage);
-	void OnLabelTextCommitted(const FText& NewText, ETextCommit::Type InTextCommit);
-	void OnEnterEditingMode();
-	void OnExitEditingMode();
+	FLinearColor GetCategoryColor() const;
+	FText GetCategoryText() const;
 
-private:
-	TWeakPtr<FWorldBookmarkTreeItem> WeakWorldBookmarkTreeItem;
+	FSlateColor GetIconColor() const;
 };
diff --git a/Engine/Source/Editor/WorldBookmark/Private/WorldBookmark/Browser/BookmarkTreeItem.cpp b/Engine/Source/Editor/WorldBookmark/Private/WorldBookmark/Browser/BookmarkTreeItem.cpp
--- a/Engine/Source/Editor/WorldBookmark/Private/WorldBookmark/Browser/BookmarkTreeItem.cpp
+++ b/Engine/Source/Editor/WorldBookmark/Private/WorldBookmark/Browser/BookmarkTreeItem.cpp
@@ -1,5 +1,83 @@
 // Copyright Epic Games, Inc. All Rights Reserved.
 
 #include "WorldBookmark/Browser/BookmarkTreeItem.h"
+#include "WorldBookmark/WorldBookmark.h"
+
+#include "AssetToolsModule.h"
+#include "AssetViewUtils.h"
+#include "ContentBrowserModule.h"
+#include "IContentBrowserSingleton.h"
+#include "ObjectTools.h"
 
 const ETreeItemType FWorldBookmarkTreeItem::Type(ETreeItemType::WorldBookmark);
+
+FName FWorldBookmarkTreeItem::GetName() const
+{
+	return BookmarkAsset.AssetName;
+}
+
+FName FWorldBookmarkTreeItem::GetAssetName() const
+{
+	return GetName();
+}
+
+static FSoftObjectPath GetRenamedObjectPath(FName PackagePath, FName InNewPackageName)
+{
+	const FString NewObjectPath = FString::Format(TEXT("{0}/{1}.{1}"), { *PackagePath.ToString(), *InNewPackageName.ToString() });
+	return FSoftObjectPath(NewObjectPath);
+}
+
+bool FWorldBookmarkTreeItem::CanRename() const
+{
+	return true;
+}
+
+bool FWorldBookmarkTreeItem::TryRename(FName InNewName, FText& OutErrorMessage) const
+{
+	const FSoftObjectPath NewObjectPath = GetRenamedObjectPath(BookmarkAsset.PackagePath, InNewName);
+	return AssetViewUtils::IsValidObjectPathForCreate(NewObjectPath.ToString(), OutErrorMessage);
+}
+
+bool FWorldBookmarkTreeItem::Rename(FName InNewName)
+{
+	FText ErrorMessage;
+	if (TryRename(InNewName, ErrorMessage))
+	{
+		FAssetToolsModule& AssetToolsModule = FModuleManager::LoadModuleChecked<FAssetToolsModule>("AssetTools");
+
+		TArray<FAssetRenameData> AssetsAndNames;
+		const FSoftObjectPath OldObjectPath = BookmarkAsset.GetSoftObjectPath();
+		const FSoftObjectPath NewObjectPath = GetRenamedObjectPath(BookmarkAsset.PackagePath, InNewName);
+		new (AssetsAndNames) FAssetRenameData(OldObjectPath, NewObjectPath);
+
+		return AssetToolsModule.Get().RenameAssetsWithDialog(AssetsAndNames) == EAssetRenameResult::Success;
+	}
+
+	return false;
+}
+
+bool FWorldBookmarkTreeItem::Delete()
+{
+	return ObjectTools::DeleteAssets({ BookmarkAsset }) == 1;
+}
+
+FName FWorldBookmarkTreeItem::GetIconName() const
+{
+	return "ClassIcon.WorldBookmark";
+}
+
+FText FWorldBookmarkTreeItem::GetText() const
+{
+	return FText::FromName(BookmarkAsset.AssetName);
+}
+
+FText FWorldBookmarkTreeItem::GetTooltipText() const
+{
+	return FText::FromString(GetAssetPath());
+}
+
+void FWorldBookmarkTreeItem::ShowInContentBrowser() const
+{
+	FContentBrowserModule& ContentBrowserModule = FModuleManager::Get().LoadModuleChecked<FContentBrowserModule>("ContentBrowser");
+	ContentBrowserModule.Get().SyncBrowserToAssets(TArray<FAssetData>({ BookmarkAsset }));
+}
\ No newline at end of file
diff --git a/Engine/Source/Editor/WorldBookmark/Private/WorldBookmark/Browser/BookmarkTreeItem.h b/Engine/Source/Editor/WorldBookmark/Private/WorldBookmark/Browser/BookmarkTreeItem.h
--- a/Engine/Source/Editor/WorldBookmark/Private/WorldBookmark/Browser/BookmarkTreeItem.h
+++ b/Engine/Source/Editor/WorldBookmark/Private/WorldBookmark/Browser/BookmarkTreeItem.h
@@ -19,15 +19,20 @@ public:
 
 	FAssetData BookmarkAsset;
 
-	virtual FString GetName() const override
-	{
-		return BookmarkAsset.AssetName.ToString();
-	}
+	virtual FName GetName() const override;
+	virtual FName GetAssetName() const override;
+	virtual bool CanRename() const override;
+	virtual bool TryRename(FName InNewName, FText& OutErrorMessage) const override;
+	virtual bool Rename(FName InNewName) override;
+	virtual bool Delete() override;
+	virtual FName GetIconName() const override;
+	virtual FText GetText() const override;
+	virtual FText GetTooltipText() const override;
+	virtual void ShowInContentBrowser() const override;
 
 	const FSoftObjectPath& GetBookmarkWorld() const;
 
 private:
-
 	mutable FSoftObjectPath CachedBookmarkWorld;
 };
 
diff --git a/Engine/Source/Editor/WorldBookmark/Private/WorldBookmark/Browser/Columns.cpp b/Engine/Source/Editor/WorldBookmark/Private/WorldBookmark/Browser/Columns.cpp
--- a/Engine/Source/Editor/WorldBookmark/Private/WorldBookmark/Browser/Columns.cpp
+++ b/Engine/Source/Editor/WorldBookmark/Private/WorldBookmark/Browser/Columns.cpp
@@ -36,6 +36,13 @@ namespace WorldBookmarkBrowserColumn
 		LOCTEXT("Column_World_Tooltip", "World Name")
 	};
 
+	const FColumnDefinition Category =
+	{
+		TEXT("Category"),
+		LOCTEXT("Column_Category_Text", "Category"),
+		LOCTEXT("Column_Category_Tooltip", "Category")
+	};
+
 	const TArray<FName> DefaultHiddenColumns = { World.Id };
 }
 
diff --git a/Engine/Source/Editor/WorldBookmark/Private/WorldBookmark/Browser/Columns.h b/Engine/Source/Editor/WorldBookmark/Private/WorldBookmark/Browser/Columns.h
--- a/Engine/Source/Editor/WorldBookmark/Private/WorldBookmark/Browser/Columns.h
+++ b/Engine/Source/Editor/WorldBookmark/Private/WorldBookmark/Browser/Columns.h
@@ -18,6 +18,7 @@ namespace WorldBookmarkBrowserColumn
 	extern const FColumnDefinition RecentlyUsed;
 	extern const FColumnDefinition Label;
 	extern const FColumnDefinition World;
+	extern const FColumnDefinition Category;
 
 	extern const TArray<FName> DefaultHiddenColumns;
 }
diff --git a/Engine/Source/Editor/WorldBookmark/Private/WorldBookmark/Browser/CommonTableRow.h b/Engine/Source/Editor/WorldBookmark/Private/WorldBookmark/Browser/CommonTableRow.h
new file mode 100644
--- /dev/null
+++ b/Engine/Source/Editor/WorldBookmark/Private/WorldBookmark/Browser/CommonTableRow.h
@@ -0,0 +1,233 @@
+// Copyright Epic Games, Inc. All Rights Reserved.
+
+#pragma once
+
+#include "Textures/SlateIcon.h"
+#include "Widgets/Images/SImage.h"
+#include "Widgets/Layout/SBox.h"
+#include "Widgets/SBoxPanel.h"
+#include "Widgets/Text/SInlineEditableTextBlock.h"
+#include "Widgets/Views/SExpanderArrow.h"
+#include "Widgets/Views/STableRow.h"
+#include "WorldBookmark/Browser/DragDropOperation.h"
+#include "WorldBookmark/Browser/TreeItem.h"
+#include "WorldBookmark/WorldBookmarkStyle.h"
+
+template <class TItemType = ITreeItem>
+class FWorldBookmarkTableRowBase
+{
+public:	
+	TSharedPtr<TItemType> GetTreeItem() const
+	{
+		if (FTreeItemPtr TreeItem = WeakTreeItem.Pin())
+		{
+			if (TreeItem->IsA<TItemType>())
+			{
+				return StaticCastSharedPtr<TItemType>(TreeItem);
+			}
+		}
+
+		return nullptr;
+	}
+
+	FText GetText() const
+	{
+		if (FTreeItemPtr TreeItem = GetTreeItem())
+		{
+			return FText::FromName(TreeItem->GetName());
+		}
+
+		return FText();
+	}
+
+	FText GetTooltipText() const
+	{
+		if (FTreeItemPtr TreeItem = GetTreeItem())
+		{
+			return FText::FromString(TreeItem->GetAssetPath());
+		}
+
+		return FText();
+	}
+
+	const FSlateBrush* GetIcon() const
+	{
+		if (FTreeItemPtr TreeItem = GetTreeItem())
+		{
+			return FSlateIcon(FWorldBookmarkStyle::Get().GetStyleSetName(), TreeItem->GetIconName()).GetIcon();
+		}
+
+		return nullptr;
+	}
+
+	FSlateColor GetIconColor() const
+	{
+		return FSlateColor::UseForeground();
+	}
+
+	FSimpleDelegate& GetRenameRequestedDelegate()
+	{
+		return WeakTreeItem.Pin()->OnRenameRequested;
+	}
+
+	bool IsReadOnly() const
+	{
+		if (FTreeItemPtr TreeItem = GetTreeItem())
+		{
+			return !TreeItem->CanRename();
+		}
+
+		return true;
+	}
+
+	bool OnLabelTextVerifyChanged(const FText& InNewText, FText& OutErrorMessage)
+	{
+		if (FTreeItemPtr TreeItem = GetTreeItem())
+		{
+			FName OldName(TreeItem->GetName());
+			FName NewName(InNewText.ToString());
+
+			return OldName == NewName || TreeItem->TryRename(FName(InNewText.ToString()), OutErrorMessage);
+		}
+
+		return false;
+	}
+
+	void OnLabelTextCommitted(const FText& InNewText, ETextCommit::Type InTextCommit)
+	{
+		if (FTreeItemPtr TreeItem = GetTreeItem())
+		{
+			FName OldName(TreeItem->GetName());
+			FName NewName(InNewText.ToString());
+
+			if (OldName != NewName)
+			{
+				TreeItem->Rename(FName(InNewText.ToString()));
+			}
+		}
+	}
+
+	void OnEnterEditingMode()
+	{
+		if (FTreeItemPtr TreeItem = GetTreeItem())
+		{
+			TreeItem->bInEditingMode = true;
+		}
+	}
+
+	void OnExitEditingMode()
+	{
+		if (FTreeItemPtr TreeItem = GetTreeItem())
+		{
+			TreeItem->bInEditingMode = false;
+		}
+	}
+
+	FReply OnRowDragDetected(const FGeometry& Geometry, const FPointerEvent& PointerEvent)
+	{
+		if (FTreeItemPtr TreeItem = GetTreeItem())
+		{
+			const TSharedPtr<FTableRowDragDropOp> DragDropOp = FTableRowDragDropOp::New(TreeItem);
+			if (DragDropOp.IsValid())
+			{
+				return FReply::Handled().BeginDragDrop(DragDropOp.ToSharedRef());
+			}
+		}
+
+		return FReply::Unhandled();
+	}
+
+	TOptional<EItemDropZone> OnRowCanAcceptDrop(const FDragDropEvent& Event, EItemDropZone Zone, TSharedPtr<TItemType> Item)
+	{
+		const TSharedPtr<FTableRowDragDropOp> DragDropOp = Event.GetOperationAs<FTableRowDragDropOp>();
+		if (DragDropOp.IsValid())
+		{
+			return DragDropOp->CanAcceptDrop(Item);
+		}
+
+		return TOptional<EItemDropZone>();
+	}
+
+	FReply OnRowAcceptDrop(const FDragDropEvent& Event, EItemDropZone Zone, TSharedPtr<TItemType> Item)
+	{
+		const TSharedPtr<FTableRowDragDropOp> DragDropOp = Event.GetOperationAs<FTableRowDragDropOp>();
+		if (DragDropOp.IsValid())
+		{
+			return DragDropOp->AcceptDrop(Item);
+		}
+
+		return FReply::Unhandled();
+	}
+		
+protected:
+	TWeakPtr<TItemType> WeakTreeItem;
+};
+
+template <class TTableRow>
+TSharedRef<SWidget> CreateEditableLabelWidget(TSharedRef<TTableRow> TableRow)
+{
+	TSharedRef<SInlineEditableTextBlock> InlineEditableTextBlock = SNew(SInlineEditableTextBlock)
+		.Text(TableRow->GetText())
+		.ToolTipText(TableRow->GetTooltipText())
+		.OnVerifyTextChanged(TableRow, &TTableRow::OnLabelTextVerifyChanged)
+		.OnTextCommitted(TableRow, &TTableRow::OnLabelTextCommitted)
+		.OnEnterEditingMode(TableRow, &TTableRow::OnEnterEditingMode)
+		.OnExitEditingMode(TableRow, &TTableRow::OnExitEditingMode)
+		.IsReadOnly(TableRow, &TTableRow::IsReadOnly)
+		.IsSelected(TableRow, &TTableRow::IsSelected);
+
+	TableRow->GetRenameRequestedDelegate().BindSP(InlineEditableTextBlock, &SInlineEditableTextBlock::EnterEditingMode);
+
+	return InlineEditableTextBlock;
+}
+
+template <class TTableRow>
+TSharedRef<SWidget> CreateTreeLabelWidget(TSharedRef<TTableRow> TableRow)
+{
+	static const int32 RowHeight = 20;
+	static const int32 IconSize = 16;
+	static const FMargin IconPadding(0.f, 1.f, 3.f, 1.f);
+
+	return SNew(SBox)
+	.MinDesiredHeight(RowHeight)
+	[
+		SNew(SHorizontalBox)
+
+		+ SHorizontalBox::Slot()
+		.AutoWidth()
+		.Padding(4, 0, 0, 0)
+		[
+			SNew(SExpanderArrow, TableRow).IndentAmount(12)
+		]
+
+		+ SHorizontalBox::Slot()
+		.FillWidth(1.0f)
+		[
+			SNew(SHorizontalBox)
+
+			+ SHorizontalBox::Slot()
+			.AutoWidth()
+			.VAlign(VAlign_Center)
+			.Padding(IconPadding)
+			[
+				SNew(SBox)
+				.WidthOverride(IconSize)
+				.HeightOverride(IconSize)
+				[
+					SNew(SImage)
+					.Image_Lambda([TableRow] { return TableRow->GetIcon(); })
+					.ToolTipText(TableRow->GetTooltipText())
+					.ColorAndOpacity(TableRow->GetIconColor())
+				]
+			]
+
+			+ SHorizontalBox::Slot()
+			.FillWidth(1.0f)
+			.VAlign(VAlign_Center)
+			.Padding(0.0f, 0.0f)
+			[
+				CreateEditableLabelWidget(TableRow)
+			]
+		]
+	];
+}
diff --git a/Engine/Source/Editor/WorldBookmark/Private/WorldBookmark/Browser/DragDropOperation.cpp b/Engine/Source/Editor/WorldBookmark/Private/WorldBookmark/Browser/DragDropOperation.cpp
new file mode 100644
--- /dev/null
+++ b/Engine/Source/Editor/WorldBookmark/Private/WorldBookmark/Browser/DragDropOperation.cpp
@@ -0,0 +1,63 @@
+// Copyright Epic Games, Inc. All Rights Reserved.
+
+#include "WorldBookmark/Browser/DragDropOperation.h"
+
+#include "Textures/SlateIcon.h"
+#include "WorldBookmark/Browser/FolderTreeItem.h"
+#include "WorldBookmark/Browser/Settings.h"
+#include "WorldBookmark/WorldBookmarkStyle.h"
+
+TSharedPtr<FTableRowDragDropOp> FTableRowDragDropOp::New(FTreeItemPtr InTreeItem)
+{
+	if (UWorldBookmarkBrowserSettings::IsViewMode(EWorldBookmarkBrowserViewMode::TreeView))
+	{
+		if (InTreeItem.IsValid())
+		{
+			if (InTreeItem->CanRename())
+			{
+				TSharedRef<FTableRowDragDropOp> Operation = MakeShared<FTableRowDragDropOp>();
+
+				Operation->CurrentHoverText = InTreeItem->GetText();
+				Operation->CurrentIconBrush = FSlateIcon(FWorldBookmarkStyle::Get().GetStyleSetName(), InTreeItem->GetIconName()).GetIcon();
+				Operation->SetupDefaults();
+
+				Operation->DraggedItem = InTreeItem;
+				Operation->Construct();
+
+				return Operation.ToSharedPtr();
+			}
+		}
+	}
+
+	return nullptr;
+}
+
+TOptional<EItemDropZone> FTableRowDragDropOp::CanAcceptDrop(FTreeItemPtr DropTarget) const
+{
+	if (DropTarget.IsValid())
+	{
+		if (FFolderTreeItem* FolderTreeItem = DropTarget->Cast<FFolderTreeItem>())
+		{
+			if (!FolderTreeItem->IsVirtual())
+			{
+				return EItemDropZone::OntoItem;
+			}
+		}
+	}
+
+	return TOptional<EItemDropZone>();
+}
+
+FReply FTableRowDragDropOp::AcceptDrop(FTreeItemPtr DropTarget) const
+{
+	if (ensure(CanAcceptDrop(DropTarget) == EItemDropZone::OntoItem))
+	{
+		if (FFolderTreeItem* FolderTreeItem = DropTarget->Cast<FFolderTreeItem>())
+		{
+			FolderTreeItem->Move(DraggedItem.Pin());
+		}
+	}
+
+	return FReply::Handled().EndDragDrop();
+}
+
diff --git a/Engine/Source/Editor/WorldBookmark/Private/WorldBookmark/Browser/DragDropOperation.h b/Engine/Source/Editor/WorldBookmark/Private/WorldBookmark/Browser/DragDropOperation.h
new file mode 100644
--- /dev/null
+++ b/Engine/Source/Editor/WorldBookmark/Private/WorldBookmark/Browser/DragDropOperation.h
@@ -0,0 +1,19 @@
+// Copyright Epic Games, Inc. All Rights Reserved.
+
+#pragma once
+
+#include "DragAndDrop/DecoratedDragDropOp.h"
+#include "Widgets/Views/STableRow.h"
+#include "WorldBookmark/Browser/TreeItem.h"
+
+class FTableRowDragDropOp : public FDecoratedDragDropOp
+{
+public:
+	DRAG_DROP_OPERATOR_TYPE(FTableRowDragDropOp, FDecoratedDragDropOp)
+	static TSharedPtr<FTableRowDragDropOp> New(FTreeItemPtr TreeItem);
+
+	TOptional<EItemDropZone> CanAcceptDrop(FTreeItemPtr DropTarget) const;
+	FReply AcceptDrop(FTreeItemPtr DropTarget) const;
+
+	TWeakPtr<ITreeItem> DraggedItem;
+};
diff --git a/Engine/Source/Editor/WorldBookmark/Private/WorldBookmark/Browser/FolderTableRow.cpp b/Engine/Source/Editor/WorldBookmark/Private/WorldBookmark/Browser/FolderTableRow.cpp
new file mode 100644
--- /dev/null
+++ b/Engine/Source/Editor/WorldBookmark/Private/WorldBookmark/Browser/FolderTableRow.cpp
@@ -0,0 +1,24 @@
+// Copyright Epic Games, Inc. All Rights Reserved.
+
+#include "WorldBookmark/Browser/FolderTableRow.h"
+#include "WorldBookmark/Browser/Icons.h"
+
+#include "Widgets/SBoxPanel.h"
+#include "Widgets/Text/SInlineEditableTextBlock.h"
+#include "Widgets/Text/STextBlock.h"
+
+void FFolderTableRow::Construct(const FArguments& InArgs, const TSharedRef<STableViewBase>& InOwnerTableView, FFolderTreeItemRef InItem)
+{
+	WeakTreeItem = InItem;
+	STableRow::Construct(
+		STableRow::FArguments(InArgs)
+		.OnDragDetected(this, &FWorldBookmarkTableRowBase::OnRowDragDetected)
+		.OnCanAcceptDrop(this, &FWorldBookmarkTableRowBase::OnRowCanAcceptDrop)
+		.OnAcceptDrop(this, &FWorldBookmarkTableRowBase::OnRowAcceptDrop),
+		InOwnerTableView);
+
+	ChildSlot
+	[
+		CreateTreeLabelWidget(SharedThis(this))
+	];
+}
diff --git a/Engine/Source/Editor/WorldBookmark/Private/WorldBookmark/Browser/FolderTableRow.h b/Engine/Source/Editor/WorldBookmark/Private/WorldBookmark/Browser/FolderTableRow.h
new file mode 100644
--- /dev/null
+++ b/Engine/Source/Editor/WorldBookmark/Private/WorldBookmark/Browser/FolderTableRow.h
@@ -0,0 +1,14 @@
+// Copyright Epic Games, Inc. All Rights Reserved.
+#pragma once
+
+#include "Widgets/Views/STableRow.h"
+#include "Widgets/Views/SListView.h"
+
+#include "WorldBookmark/Browser/FolderTreeItem.h"
+#include "WorldBookmark/Browser/CommonTableRow.h"
+
+class FFolderTableRow : public STableRow<FFolderTreeItemPtr>, public FWorldBookmarkTableRowBase<FFolderTreeItem>
+{
+public:
+	void Construct(const FArguments& InArgs, const TSharedRef<STableViewBase>& InOwnerTableView, FFolderTreeItemRef InItem);
+};
diff --git a/Engine/Source/Editor/WorldBookmark/Private/WorldBookmark/Browser/FolderTreeItem.cpp b/Engine/Source/Editor/WorldBookmark/Private/WorldBookmark/Browser/FolderTreeItem.cpp
new file mode 100644
--- /dev/null
+++ b/Engine/Source/Editor/WorldBookmark/Private/WorldBookmark/Browser/FolderTreeItem.cpp
@@ -0,0 +1,297 @@
+// Copyright Epic Games, Inc. All Rights Reserved.
+
+#include "WorldBookmark/Browser/FolderTreeItem.h"
+
+#include "WorldBookmark/Browser/BookmarkTreeItem.h"
+
+#include "AssetToolsModule.h"
+#include "AssetViewUtils.h"
+#include "ContentBrowserModule.h"
+#include "IContentBrowserSingleton.h"
+#include "ObjectTools.h"
+
+const ETreeItemType FFolderTreeItem::Type(ETreeItemType::Folder);
+
+
+FFolderTreeItem::FFolderTreeItem(const FName InName)
+	: ITreeItem(Type)
+	, Name(InName)
+	, bIsExpanded(true)
+	, bIsVirtual(false)
+	, bIsMountPoint(false)
+{
+}
+
+FName FFolderTreeItem::GetName() const
+{
+	return Name;
+}
+
+FName FFolderTreeItem::GetAssetName() const
+{
+	if (bIsVirtual)
+	{
+		return NAME_None;
+	}
+	else if (bIsMountPoint)
+	{
+		return MountPointName;
+	}
+	else
+	{
+		return Name;
+	}
+}
+
+bool FFolderTreeItem::CanRename() const
+{
+	return !IsVirtual() && !IsMountPoint();
+}
+
+bool FFolderTreeItem::TryRename(FName InNewName, FText& OutErrorMessage) const
+{
+	return AssetViewUtils::IsValidFolderPathForCreate(GetAssetPath(), InNewName.ToString(), OutErrorMessage);
+}
+
+bool FFolderTreeItem::Rename(FName InNewName)
+{
+	check(!bIsVirtual && !bIsMountPoint);
+
+	FString OldPathToFolder = GetAssetPath();
+	FString NewPathToFolder = (GetParent() ? GetParent()->GetAssetPath() : FString()) / InNewName.ToString();
+
+	TArray<FAssetRenameData> AssetsAndNames;
+
+	ForEachChildRecursive<FWorldBookmarkTreeItem>([&AssetsAndNames, &OldPathToFolder, &NewPathToFolder](FWorldBookmarkTreeItemPtr WorldBookmarkItem)
+	{
+		FSoftObjectPath OldObjectPath = WorldBookmarkItem->BookmarkAsset.GetSoftObjectPath();
+		FSoftObjectPath NewObjectPath = OldObjectPath.ToString().Replace(*OldPathToFolder, *NewPathToFolder);
+		new (AssetsAndNames) FAssetRenameData(OldObjectPath, NewObjectPath);
+	});	
+
+	FAssetToolsModule& AssetToolsModule = FModuleManager::LoadModuleChecked<FAssetToolsModule>("AssetTools");
+	return AssetToolsModule.Get().RenameAssetsWithDialog(AssetsAndNames) == EAssetRenameResult::Success;
+}
+
+bool FFolderTreeItem::Move(FTreeItemPtr NewChild)
+{
+	TArray<FAssetRenameData> AssetsAndNames;
+
+	FString OldPathToItem = NewChild->GetAssetPath();
+	FString NewPathToItem = GetAssetPath() / NewChild->GetName().ToString();
+
+	auto AddAssetToRename = [&AssetsAndNames, &OldPathToItem, &NewPathToItem](FWorldBookmarkTreeItem* WorldBookmarkItem)
+	{
+		FSoftObjectPath OldObjectPath = WorldBookmarkItem->BookmarkAsset.GetSoftObjectPath();
+		FSoftObjectPath NewObjectPath = OldObjectPath.ToString().Replace(*OldPathToItem, *NewPathToItem);
+		new (AssetsAndNames) FAssetRenameData(OldObjectPath, NewObjectPath);
+	};
+
+	if (FWorldBookmarkTreeItem* WorldBookmarkItem = NewChild->Cast<FWorldBookmarkTreeItem>())
+	{
+		AddAssetToRename(WorldBookmarkItem);
+	}
+	else if (FFolderTreeItem* FolderTreeItem = NewChild->Cast<FFolderTreeItem>())
+	{
+		FolderTreeItem->ForEachChildRecursive<FWorldBookmarkTreeItem>([&AddAssetToRename](FWorldBookmarkTreeItemPtr WorldBookmarkItem)
+		{
+			AddAssetToRename(WorldBookmarkItem.Get());
+		});
+	}
+
+	FAssetToolsModule& AssetToolsModule = FModuleManager::LoadModuleChecked<FAssetToolsModule>("AssetTools");
+	return AssetToolsModule.Get().RenameAssetsWithDialog(AssetsAndNames) == EAssetRenameResult::Success;
+}
+
+bool FFolderTreeItem::Delete()
+{
+	TArray<FAssetData> AssetsToDelete;
+
+	ForEachChildRecursive<FWorldBookmarkTreeItem>([&AssetsToDelete](FWorldBookmarkTreeItemPtr WorldBookmarkItem)
+	{
+		AssetsToDelete.Add(WorldBookmarkItem->BookmarkAsset);
+	});
+
+	return ObjectTools::DeleteAssets({ AssetsToDelete }) == AssetsToDelete.Num();
+}
+
+FName FFolderTreeItem::GetIconName() const
+{
+	return bIsExpanded ? "WorldBookmark.FolderOpen" : "WorldBookmark.FolderClosed";
+}
+
+FText FFolderTreeItem::GetText() const
+{
+	return FText::FromName(Name);
+}
+
+FText FFolderTreeItem::GetTooltipText() const
+{
+	return FText::FromString(GetAssetPath());
+}
+
+void FFolderTreeItem::ShowInContentBrowser() const
+{
+	FContentBrowserModule& ContentBrowserModule = FModuleManager::Get().LoadModuleChecked<FContentBrowserModule>("ContentBrowser");
+	ContentBrowserModule.Get().SyncBrowserToFolders({ GetAssetPath() });
+}
+
+const TArray<FTreeItemPtr>& FFolderTreeItem::GetChildren() const
+{
+	return Children;
+}
+
+void FFolderTreeItem::AddChild(FTreeItemPtr InChildItem)
+{
+	Children.Add(InChildItem);
+	InChildItem->SetParent(this);
+}
+
+FFolderTreeItemPtr FFolderTreeItem::CreateMountPoint(const FString& InPath, const FString& InMountPointName)
+{
+	FFolderTreeItemPtr MountPoint = CreatePath(InPath, true);
+	MountPoint->bIsMountPoint = true;
+	MountPoint->MountPointName = FName(InMountPointName);
+	return MountPoint;
+}
+
+FFolderTreeItemPtr FFolderTreeItem::CreatePath(const FString& InPath)
+{
+	return CreatePath(InPath, false);
+}
+
+FFolderTreeItemPtr FFolderTreeItem::CreatePath(const FString& InPath, bool bInIsMountPoint)
+{
+	TArray<FString> PathAsStringArray;
+	InPath.ParseIntoArray(PathAsStringArray, TEXT("/"), true);
+
+	TArray<FName> PathAsNameArray;
+	Algo::Transform(PathAsStringArray, PathAsNameArray, [](const FString& InString) { return FName(InString); });
+	return CreatePath(PathAsNameArray, bInIsMountPoint);
+}
+
+FFolderTreeItemPtr FFolderTreeItem::CreatePath(const TArrayView<FName>& InPath, bool bInIsCreatingMountPoint)
+{
+	if (InPath.IsEmpty())
+	{
+		return StaticCastSharedPtr<FFolderTreeItem>(AsShared().ToSharedPtr());
+	}
+
+	FFolderTreeItemPtr NextFolder = nullptr;
+	for (TSharedPtr<ITreeItem>& Child : Children)
+	{
+		if (Child->IsA<FFolderTreeItem>() && Child->GetName() == InPath[0])
+		{
+			NextFolder = StaticCastSharedPtr<FFolderTreeItem>(Child);
+			break;
+		}
+	}
+
+	if (NextFolder == nullptr)
+	{
+		NextFolder = MakeShared<FFolderTreeItem>(InPath[0]);
+		NextFolder->SetParent(this);
+		Children.Add(NextFolder);
+	}
+
+	if (InPath.Num() == 1)
+	{
+		// We're done, that was the last folder in the path
+		return NextFolder;
+	}
+	else
+	{
+		// Folders above a mount point does not exists, mark them as "virtual"
+		NextFolder->bIsVirtual = bInIsCreatingMountPoint;
+
+		// Recurse deeper in the tree
+		return NextFolder->CreatePath(InPath.RightChop(1), bInIsCreatingMountPoint);
+	}
+}
+
+void FFolderTreeItem::ExpandPath(const FName InPath) const
+{
+	TArray<FString> PathAsStringArray;
+	InPath.ToString().ParseIntoArray(PathAsStringArray, TEXT("/"), true);
+	
+	TArray<FName> PathAsNameArray;
+	Algo::Transform(PathAsStringArray, PathAsNameArray, [](const FString& InString) { return FName(InString); });
+	return ExpandPath(PathAsNameArray);
+}
+
+void FFolderTreeItem::ExpandPath(const TArrayView<FName>& InPath) const
+{
+
+}
+
+bool FFolderTreeItem::IsExpanded() const
+{
+	return bIsExpanded;
+}
+
+void FFolderTreeItem::SetExpanded(bool bInExpanded)
+{
+	bIsExpanded = bInExpanded;
+}
+
+bool FFolderTreeItem::IsVirtual() const
+{
+	return bIsVirtual;
+}
+
+bool FFolderTreeItem::IsMountPoint() const
+{
+	return bIsMountPoint;
+}
+
+void FFolderTreeItem::ClearBookmarkItems()
+{
+	for (TArray<FTreeItemPtr>::TIterator It = Children; It; ++It)
+	{
+		TSharedPtr<ITreeItem>& Child = *It;
+		if (FWorldBookmarkTreeItem* BookmarkTreeItem = Child->Cast<FWorldBookmarkTreeItem>())
+		{
+			It.RemoveCurrent();
+		}
+		else if (FFolderTreeItem* ChildFolder = Child->Cast<FFolderTreeItem>())
+		{
+			ChildFolder->ClearBookmarkItems();
+		}
+	}
+}
+
+bool FFolderTreeItem::ClearEmptyFolders()
+{
+	for (TArray<FTreeItemPtr>::TIterator It = Children; It; ++It)
+	{
+		TSharedPtr<ITreeItem>& Child = *It;
+		if (FFolderTreeItem* ChildFolder = Child->Cast<FFolderTreeItem>())
+		{
+			bool bIsEmpty = ChildFolder->ClearEmptyFolders();
+			if (bIsEmpty)
+			{
+				It.RemoveCurrent();
+			}
+		}
+	}
+
+	return Children.IsEmpty();
+}
+
+void FFolderTreeItem::Sort(TFunctionRef<bool(const FTreeItemPtr& A, const FTreeItemPtr& B)> SortFunc)
+{
+	Children.Sort(SortFunc);
+
+	for (TSharedPtr<ITreeItem>& Child : Children)
+	{
+		if (FFolderTreeItem* ChildFolder = Child->Cast<FFolderTreeItem>())
+		{
+			ChildFolder->Sort(SortFunc);
+		}
+	}
+}
+
+void FFolderTreeItem::Reset()
+{
+	Children.Reset();
+}
\ No newline at end of file
diff --git a/Engine/Source/Editor/WorldBookmark/Private/WorldBookmark/Browser/FolderTreeItem.h b/Engine/Source/Editor/WorldBookmark/Private/WorldBookmark/Browser/FolderTreeItem.h
new file mode 100644
--- /dev/null
+++ b/Engine/Source/Editor/WorldBookmark/Private/WorldBookmark/Browser/FolderTreeItem.h
@@ -0,0 +1,90 @@
+// Copyright Epic Games, Inc. All Rights Reserved.
+#pragma once
+
+#include "WorldBookmark/Browser/TreeItem.h"
+
+struct FFolderTreeItem : public ITreeItem
+{
+public:
+	/** Static type identifier for this tree item class */
+	static const ETreeItemType Type;
+
+	FFolderTreeItem(const FName InName = NAME_None);
+
+	virtual FName GetName() const override;
+	virtual FName GetAssetName() const override;
+	virtual bool CanRename() const override;
+	virtual bool TryRename(FName InNewName, FText& OutErrorMessage) const override;
+	virtual bool Rename(FName InNewName) override;
+	virtual bool Delete() override;
+	virtual FName GetIconName() const override;
+	virtual FText GetText() const override;
+	virtual FText GetTooltipText() const override;
+	virtual void ShowInContentBrowser() const override;
+
+	const TArray<FTreeItemPtr>& GetChildren() const;
+	void AddChild(FTreeItemPtr InChildItem);
+
+	TSharedPtr<FFolderTreeItem> CreateMountPoint(const FString& InPath, const FString& InMountPointName);
+	TSharedPtr<FFolderTreeItem> CreatePath(const FString& InPath);
+	void ExpandPath(FName InPath) const;
+
+	bool IsExpanded() const;
+	void SetExpanded(bool bInExpanded);
+
+	bool IsVirtual() const;
+	bool IsMountPoint() const;
+
+	void Sort(TFunctionRef<bool(const FTreeItemPtr&, const FTreeItemPtr&)> SortFunc);
+
+	void Reset();
+
+	void ClearBookmarkItems();
+	bool ClearEmptyFolders();
+
+	bool Move(FTreeItemPtr NewChild);
+
+	template <typename T>
+	void ForEachChild(TFunctionRef<void(TSharedPtr<T>)> InFunc)
+	{
+		for (FTreeItemPtr& Child : Children)
+		{
+			if (Child->IsA<T>())
+			{
+				InFunc(StaticCastSharedPtr<T>(Child));
+			}
+		}
+	}
+
+	template <typename T>
+	void ForEachChildRecursive(TFunctionRef<void(TSharedPtr<T>)> InFunc)
+	{
+		for (FTreeItemPtr& Child : Children)
+		{
+			if (Child->IsA<T>())
+			{
+				InFunc(StaticCastSharedPtr<T>(Child));
+			}
+
+			if (FFolderTreeItem* Folder = Child->Cast<FFolderTreeItem>())
+			{
+				Folder->ForEachChildRecursive(InFunc);
+			}
+		}
+	}
+
+private:
+	TSharedPtr<FFolderTreeItem> CreatePath(const FString& InPath, bool bIsCreatingMountPoint);
+	TSharedPtr<FFolderTreeItem> CreatePath(const TArrayView<FName>& InPath, bool bIsCreatingMountPoint);
+	void ExpandPath(const TArrayView<FName>& InPath) const;
+
+	FName Name;
+	FName MountPointName;
+	bool bIsExpanded;
+	bool bIsVirtual;
+	bool bIsMountPoint;
+	TArray<FTreeItemPtr> Children;
+};
+
+typedef TSharedPtr<FFolderTreeItem> FFolderTreeItemPtr;
+typedef TSharedRef<FFolderTreeItem> FFolderTreeItemRef;
diff --git a/Engine/Source/Editor/WorldBookmark/Private/WorldBookmark/Browser/Icons.cpp b/Engine/Source/Editor/WorldBookmark/Private/WorldBookmark/Browser/Icons.cpp
--- a/Engine/Source/Editor/WorldBookmark/Private/WorldBookmark/Browser/Icons.cpp
+++ b/Engine/Source/Editor/WorldBookmark/Private/WorldBookmark/Browser/Icons.cpp
@@ -4,6 +4,7 @@
 
 #include "Widgets/Images/SImage.h"
 #include "Widgets/Layout/SBox.h"
+#include "Widgets/SBoxPanel.h"
 #include "WorldBookmark/WorldBookmarkStyle.h"
 
 FSlateIcon GetFavoriteWorldBookmarkIcon(bool bIsFavorite)
@@ -18,15 +19,26 @@ FSlateIcon GetRecentlyUsedWorldBookmarkIcon()
 	return FSlateIcon(FWorldBookmarkStyle::Get().GetStyleSetName(), RecentlyUsedIconName);
 }
 
+FSlateIcon GetFolderIcon(bool bIsExpanded)
+{
+	const FName FolderIconName = bIsExpanded ? "WorldBookmark.FolderOpen" : "WorldBookmark.FolderClosed";
+	return FSlateIcon(FWorldBookmarkStyle::Get().GetStyleSetName(), FolderIconName);
+}
+
+FSlateIcon GetWorldBookmarkIcon()
+{
+	const FName WorldBookmarkIconName = "ClassIcon.WorldBookmark";
+	return FSlateIcon(FWorldBookmarkStyle::Get().GetStyleSetName(), WorldBookmarkIconName);
+}
+
 TSharedRef<SWidget> CreateIconWidget(const FSlateIcon& SlateIcon)
 {
 	return SNew(SBox)
-		.Padding(2.0f)
 		.HAlign(HAlign_Center)
 		.VAlign(VAlign_Center)
 		[
 			SNew(SImage)
-				.DesiredSizeOverride(FVector2D(20, 20))
+				.DesiredSizeOverride(FVector2D(16, 16))
 				.ColorAndOpacity(FSlateColor::UseForeground())
 				.Image(SlateIcon.GetIcon())
 		];
diff --git a/Engine/Source/Editor/WorldBookmark/Private/WorldBookmark/Browser/Icons.h b/Engine/Source/Editor/WorldBookmark/Private/WorldBookmark/Browser/Icons.h
--- a/Engine/Source/Editor/WorldBookmark/Private/WorldBookmark/Browser/Icons.h
+++ b/Engine/Source/Editor/WorldBookmark/Private/WorldBookmark/Browser/Icons.h
@@ -8,4 +8,6 @@ class SWidget;
 
 FSlateIcon GetFavoriteWorldBookmarkIcon(bool bIsFavorite);
 FSlateIcon GetRecentlyUsedWorldBookmarkIcon();
+FSlateIcon GetFolderIcon(bool bIsExpanded);
+FSlateIcon GetWorldBookmarkIcon();
 TSharedRef<SWidget> CreateIconWidget(const FSlateIcon& SlateIcon);
diff --git a/Engine/Source/Editor/WorldBookmark/Private/WorldBookmark/Browser/SWorldBookmarkBrowser.cpp b/Engine/Source/Editor/WorldBookmark/Private/WorldBookmark/Browser/SWorldBookmarkBrowser.cpp
--- a/Engine/Source/Editor/WorldBookmark/Private/WorldBookmark/Browser/SWorldBookmarkBrowser.cpp
+++ b/Engine/Source/Editor/WorldBookmark/Private/WorldBookmark/Browser/SWorldBookmarkBrowser.cpp
@@ -7,6 +7,8 @@
 #include "WorldBookmark/Browser/Settings.h"
 #include "WorldBookmark/Browser/BookmarkTreeItem.h"
 #include "WorldBookmark/Browser/BookmarkTableRow.h"
+#include "WorldBookmark/Browser/FolderTreeItem.h"
+#include "WorldBookmark/Browser/FolderTableRow.h"
 
 #include "WorldBookmark/WorldBookmark.h"
 #include "WorldBookmark/WorldBookmarkCommands.h"
@@ -18,17 +20,17 @@
 #include "AssetRegistry/AssetData.h"
 #include "AssetRegistry/AssetRegistryHelpers.h"
 #include "AssetToolsModule.h"
-#include "ContentBrowserDataSource.h"
 #include "ContentBrowserModule.h"
 #include "Editor.h"
 #include "Framework/Application/SlateApplication.h"
 #include "Framework/Commands/GenericCommands.h"
+#include "IContentBrowserDataModule.h"
+#include "Interfaces/IPluginManager.h"
 #include "Kismet2/DebuggerCommands.h"
 #include "LevelEditor.h"
 #include "LevelEditorCameraEditorState.h"
 #include "Misc/MessageDialog.h"
 #include "Modules/ModuleManager.h"
-#include "ObjectTools.h"
 #include "PropertyEditorModule.h"
 #include "SLevelViewport.h"
 #include "SourceControlHelpers.h"
@@ -36,6 +38,7 @@
 #include "Toolkits/GlobalEditorCommonCommands.h"
 #include "UncontrolledChangelistsModule.h"
 #include "UObject/AssetRegistryTagsContext.h"
+#include "Widgets/Colors/SColorBlock.h"
 #include "Widgets/Input/SComboButton.h"
 #include "Widgets/Input/SSearchBox.h"
 #include "Widgets/Input/SSpinBox.h"
@@ -59,6 +62,7 @@ static UWorld* GetCurrentWorld()
 SWorldBookmarkBrowser::SWorldBookmarkBrowser()
 	: bIsInTick(false)
 	, bPendingRefresh(false)
+	, bExpandAllItemsOnNextRefresh(false)
 {
 	// Listen for asset registry updates
 	IAssetRegistry& AssetRegistry = IAssetRegistry::GetChecked();
@@ -100,7 +104,7 @@ SWorldBookmarkBrowser::~SWorldBookmarkBrowser()
 
 void SWorldBookmarkBrowser::Construct(const FArguments& InArgs)
 {
-	RefreshItems();
+	TreeRoot = MakeShared<FFolderTreeItem>();
 
 	// Create a property view
 	FPropertyEditorModule& EditModule = FModuleManager::Get().GetModuleChecked<FPropertyEditorModule>("PropertyEditor");
@@ -110,49 +114,16 @@ void SWorldBookmarkBrowser::Construct(const FArguments& InArgs)
 	PropertyView = EditModule.CreateDetailView(DetailsViewArgs);
 	PropertyView->SetObject(nullptr, true);
 
-	auto HeaderRowColumn = [this](const WorldBookmarkBrowserColumn::FColumnDefinition& InColumn) -> SHeaderRow::FColumn::FArguments
-	{
-		return SHeaderRow::Column(InColumn.Id)
-			.SortMode(this, &SWorldBookmarkBrowser::GetColumnSortMode, InColumn.Id)
-			.SortPriority(this, &SWorldBookmarkBrowser::GetColumnSortPriority, InColumn.Id)
-			.OnSort(this, &SWorldBookmarkBrowser::OnColumnSortChanged)
-			.DefaultLabel(InColumn.DisplayText)
-			.DefaultTooltip(InColumn.ToolTipText);
-	};
-
-	TSharedRef<SHeaderRow> HeaderRow = SNew(SHeaderRow)
+	HeaderRow = SNew(SHeaderRow)
 		.CanSelectGeneratedColumn(true)
 		.HiddenColumnsList(UWorldBookmarkBrowserSettings::GetHiddenColumns())
 		.OnHiddenColumnsListChanged_Lambda([this]()
 		{
 			UWorldBookmarkBrowserSettings::SetHiddenColumns(BookmarksView->GetHeaderRow()->GetHiddenColumnIds());
-		})
-
-		// Favorite
-		+ HeaderRowColumn(WorldBookmarkBrowserColumn::Favorite)
-		.FixedWidth(32)
-		.HeaderContent()
-		[
-			CreateIconWidget(GetFavoriteWorldBookmarkIcon(true))
-		]
-
-		// RecentlyUsed
-		+ HeaderRowColumn(WorldBookmarkBrowserColumn::RecentlyUsed)
-		.FixedWidth(32)
-		.HeaderContent()
-		[
-			CreateIconWidget(GetRecentlyUsedWorldBookmarkIcon())
-		]
-
-		// Label
-		+ HeaderRowColumn(WorldBookmarkBrowserColumn::Label)
-		.ShouldGenerateWidget(true) // Ensure the column cannot be hidden (grayed out in the show/hide drop down menu)
-
-		// World
-		+ HeaderRowColumn(WorldBookmarkBrowserColumn::World);
+		});
 
 	BookmarksView = SNew(STreeView<FTreeItemPtr>)
-		.TreeItemsSource(&TreeItems)
+		.TreeItemsSource(&TreeRoot->GetChildren())
 		.SelectionMode(ESelectionMode::Single)
 		.OnGetChildren(this, &SWorldBookmarkBrowser::OnGetChildren)
 		.OnGenerateRow(this, &SWorldBookmarkBrowser::OnGenerateRow)
@@ -160,6 +131,8 @@ void SWorldBookmarkBrowser::Construct(const FArguments& InArgs)
 		.OnMouseButtonDoubleClick(this, &SWorldBookmarkBrowser::OnMouseButtonDoubleClicked)
 		.OnContextMenuOpening(this, &SWorldBookmarkBrowser::OnContextMenuOpening)
 		.OnItemScrolledIntoView(this, &SWorldBookmarkBrowser::OnItemScrolledIntoView)
+		.OnExpansionChanged(this, &SWorldBookmarkBrowser::OnItemExpansionChanged)
+		.OnSetExpansionRecursive(this, &SWorldBookmarkBrowser::SetItemExpansionRecursive)
 		.HeaderRow
 		(
 			HeaderRow
@@ -242,16 +215,16 @@ void SWorldBookmarkBrowser::Construct(const FArguments& InArgs)
 	Commands = MakeShareable(new FUICommandList());
 
 	Commands->MapAction(FGlobalEditorCommonCommands::Get().FindInContentBrowser,
-		FExecuteAction::CreateSP(this, &SWorldBookmarkBrowser::FindSelectedBookmarkInContentBrowser),
-		FCanExecuteAction::CreateSP(this, &SWorldBookmarkBrowser::IsValidBookmarkSelected));
+		FExecuteAction::CreateSP(this, &SWorldBookmarkBrowser::FindSelectedItemInContentBrowser),
+		FCanExecuteAction::CreateSP(this, &SWorldBookmarkBrowser::IsValidItemSelected));
 
 	Commands->MapAction(FGenericCommands::Get().Rename,
-		FExecuteAction::CreateSP(this, &SWorldBookmarkBrowser::RequestRenameSelectedBookmark),
-		FCanExecuteAction::CreateSP(this, &SWorldBookmarkBrowser::IsValidBookmarkSelected));
+		FExecuteAction::CreateSP(this, &SWorldBookmarkBrowser::RequestRenameSelectedItem),
+		FCanExecuteAction::CreateSP(this, &SWorldBookmarkBrowser::CanRenameOrDeleteSelectedItem));
 
 	Commands->MapAction(FGenericCommands::Get().Delete,
-		FExecuteAction::CreateSP(this, &SWorldBookmarkBrowser::DeleteSelectedBookmark),
-		FCanExecuteAction::CreateSP(this, &SWorldBookmarkBrowser::IsValidBookmarkSelected));
+		FExecuteAction::CreateSP(this, &SWorldBookmarkBrowser::DeleteSelectedItem),
+		FCanExecuteAction::CreateSP(this, &SWorldBookmarkBrowser::CanRenameOrDeleteSelectedItem));
 
 	Commands->MapAction(FWorldBookmarkCommands::Get().LoadBookmark,
 		FExecuteAction::CreateSP(this, &SWorldBookmarkBrowser::LoadSelectedBookmark),
@@ -281,13 +254,142 @@ void SWorldBookmarkBrowser::Construct(const FArguments& InArgs)
 		FExecuteAction::CreateSP(this, &SWorldBookmarkBrowser::GoToSelectedBookmarkLocation),
 		FCanExecuteAction::CreateSP(this, &SWorldBookmarkBrowser::CanGoToSelectedBookmarkLocation));
 
+	Commands->MapAction(FWorldBookmarkCommands::Get().MoveBookmarkToNewFolder,
+		FExecuteAction::CreateSP(this, &SWorldBookmarkBrowser::MoveSelectedBookmarkToNewFolder),
+		FCanExecuteAction(),
+		FIsActionChecked(),
+		FIsActionButtonVisible::CreateLambda([this]() { return GetSelectedBookmark() != nullptr; }));
+
+	Commands->MapAction(FWorldBookmarkCommands::Get().CreateBookmarkInFolder,
+		FExecuteAction::CreateSP(this, &SWorldBookmarkBrowser::CreateBookmarkInSelectedFolder),
+		FCanExecuteAction(),
+		FIsActionChecked(),
+		FIsActionButtonVisible::CreateLambda([this]() { return GetSelectedTreeItem() != nullptr; }));
+
 	UpdateDetailsView(/*InSelectedWorldBookmark=*/nullptr);
+
+	RefreshItems(/*bForceRefresh=*/true);
+
+	RecreateColumns();
+}
+
+void SWorldBookmarkBrowser::RecreateColumns()
+{
+	HeaderRow->ClearColumns();
+
+	auto HeaderRowColumn = [this](const WorldBookmarkBrowserColumn::FColumnDefinition& InColumn) -> SHeaderRow::FColumn::FArguments
+	{
+		return SHeaderRow::Column(InColumn.Id)
+			.SortMode(this, &SWorldBookmarkBrowser::GetColumnSortMode, InColumn.Id)
+			.SortPriority(this, &SWorldBookmarkBrowser::GetColumnSortPriority, InColumn.Id)
+			.OnSort(this, &SWorldBookmarkBrowser::OnColumnSortChanged)
+			.DefaultLabel(InColumn.DisplayText)
+			.DefaultTooltip(InColumn.ToolTipText);
+	};
+
+	auto HeaderRowIconColumn = [this, &HeaderRowColumn](const WorldBookmarkBrowserColumn::FColumnDefinition& InColumn)->SHeaderRow::FColumn::FArguments
+	{
+		return HeaderRowColumn(InColumn)
+			.FixedWidth(24.f)
+			.HAlignHeader(HAlign_Left)
+			.VAlignHeader(VAlign_Center)
+			.HAlignCell(HAlign_Center)
+			.VAlignCell(VAlign_Center)
+			.HeaderContentPadding(FMargin(4.0f, 0.0f, 0.0f, 0.0f));
+	};
+	
+
+	// Favorite
+	SHeaderRow::FColumn::FArguments FavoriteColumn =
+		HeaderRowIconColumn(WorldBookmarkBrowserColumn::Favorite)
+		[
+			CreateIconWidget(GetFavoriteWorldBookmarkIcon(true))
+		];
+
+	// RecentlyUsed
+	SHeaderRow::FColumn::FArguments RecentlyUsedColumn =
+		HeaderRowIconColumn(WorldBookmarkBrowserColumn::RecentlyUsed)
+		[
+			CreateIconWidget(GetRecentlyUsedWorldBookmarkIcon())
+		];
+
+	// LabelColumn
+	SHeaderRow::FColumn::FArguments LabelColumn =
+		HeaderRowColumn(WorldBookmarkBrowserColumn::Label)
+		.ShouldGenerateWidget(true); // Ensure the column cannot be hidden (grayed out in the show/hide drop down menu)
+
+	// World
+	SHeaderRow::FColumn::FArguments WorldColumn =
+		HeaderRowColumn(WorldBookmarkBrowserColumn::World)
+		.FillWidth(0.65f);
+
+	// Category
+	SHeaderRow::FColumn::FArguments CategoryColumn =
+		HeaderRowIconColumn(WorldBookmarkBrowserColumn::Category)
+		[
+			SNew(SColorBlock)
+				.Color(FStyleColors::Foreground.GetSpecifiedColor())
+				.CornerRadius(FVector4(4.0f, 4.0f, 4.0f, 4.0f))
+				.IsEnabled(false)
+				.Size(FVector2D(14,14))
+		];
+
+	// Order columns differently based on the view mode
+	if (UWorldBookmarkBrowserSettings::IsViewMode(EWorldBookmarkBrowserViewMode::TreeView))
+	{
+		HeaderRow->AddColumn(LabelColumn);
+		HeaderRow->AddColumn(WorldColumn);
+		HeaderRow->AddColumn(CategoryColumn);
+		HeaderRow->AddColumn(FavoriteColumn);
+		HeaderRow->AddColumn(RecentlyUsedColumn);
+	}
+	else
+	{
+		HeaderRow->AddColumn(FavoriteColumn);
+		HeaderRow->AddColumn(RecentlyUsedColumn);
+		HeaderRow->AddColumn(CategoryColumn);
+		HeaderRow->AddColumn(LabelColumn);
+		HeaderRow->AddColumn(WorldColumn);
+	}
 }
 
 TSharedRef<SWidget> SWorldBookmarkBrowser::GetSettingsMenuContent()
 {
 	FMenuBuilder MenuBuilder(false, nullptr);
 
+	// ViewType Section
+	MenuBuilder.BeginSection(TEXT("SettingsMenuSection_ViewType"), LOCTEXT("SettingsMenuSection_ViewType", "View Type"));
+	{
+		// View Type -> List View
+		MenuBuilder.AddMenuEntry(
+			LOCTEXT("Settings_ViewType_ListView", "List"),
+			LOCTEXT("Settings_ViewType_ListView_Tooltip", "View bookmarks as a list"),
+			FSlateIcon(),
+			FUIAction(
+				FExecuteAction::CreateStatic(&UWorldBookmarkBrowserSettings::SetViewMode, EWorldBookmarkBrowserViewMode::ListView),
+				FCanExecuteAction(),
+				FIsActionChecked::CreateStatic(&UWorldBookmarkBrowserSettings::IsViewMode, EWorldBookmarkBrowserViewMode::ListView)
+			),
+			NAME_None,
+			EUserInterfaceActionType::RadioButton
+		);
+
+		// View Type -> Tree View
+		MenuBuilder.AddMenuEntry(
+			LOCTEXT("Settings_ViewType_TreeView", "Tree"),
+			LOCTEXT("Settings_ViewType_TreeView_Tooltip", "View bookmarks as a tree based on the asset path"),
+			FSlateIcon(),
+			FUIAction(
+				FExecuteAction::CreateStatic(&UWorldBookmarkBrowserSettings::SetViewMode, EWorldBookmarkBrowserViewMode::TreeView),
+				FCanExecuteAction(),
+				FIsActionChecked::CreateStatic(&UWorldBookmarkBrowserSettings::IsViewMode, EWorldBookmarkBrowserViewMode::TreeView)
+			),
+			NAME_None,
+			EUserInterfaceActionType::RadioButton
+		);
+	}
+	MenuBuilder.EndSection();
+
 	// World Bookmark
 	MenuBuilder.BeginSection(TEXT("SettingsMenuSection_Show"), LOCTEXT("SettingsMenuSection_Show", "Show"));
 	{
@@ -394,8 +496,6 @@ void SWorldBookmarkBrowser::RefreshItems(bool bForceRefresh)
 
 	UWorldBookmark* PreviousSelection = GetSelectedBookmark();
 
-	TreeItems.Reset();
-
 	FSoftObjectPath CurrentWorldPath = FSoftObjectPath(GetCurrentWorld());
 	
 	TArray<FUncontrolledChangelistStateRef> UncontrolledChangelistsStates;
@@ -528,35 +628,79 @@ void SWorldBookmarkBrowser::RefreshItems(bool bForceRefresh)
 		AssetsData.SetNum(UWorldBookmarkBrowserSettings::GetMaxLastRecentlyUsedItems());
 	}
 
+	// Recreate the tree items, but retain the expansion state of folders
+	TreeRoot->ClearBookmarkItems();
+	
+	BookmarkTreeItems.Reset();
+
 	for (const FAssetData& AssetData : AssetsData)
 	{
-		TreeItems.Emplace_GetRef(MakeShared<FWorldBookmarkTreeItem>(AssetData));
+		FWorldBookmarkTreeItemPtr WorldBookmarkTreeItem = MakeShared<FWorldBookmarkTreeItem>(AssetData);
+		BookmarkTreeItems.Emplace(AssetData.PackageName, WorldBookmarkTreeItem);
+
+		// In flat view mode, add everything to the root
+		if (UWorldBookmarkBrowserSettings::GetViewMode() == EWorldBookmarkBrowserViewMode::ListView)
+		{
+			TreeRoot->AddChild(WorldBookmarkTreeItem);
+		}
+		else
+		{
+			FString PackagePath = AssetData.PackagePath.ToString();
+
+			// Extract the mount point
+			const FStringView MountPoint = FPathViews::GetMountPointNameFromPath(PackagePath);
+
+			// Convert the mount point to a virtual path (ex: /Game -> /All/Content or /MyPlugin/ -> /All/Plugins/MyPlugin)
+			UContentBrowserDataSubsystem* ContentBrowserDataSubsystem = IContentBrowserDataModule::Get().GetSubsystem();			
+			FName VirtualPathName;
+			ContentBrowserDataSubsystem->ConvertInternalPathToVirtual(FString("/") + MountPoint, VirtualPathName);
+		
+			// Create the mount point folder
+			FFolderTreeItemPtr MountPointFolder = TreeRoot->CreateMountPoint(VirtualPathName.ToString(), FString(MountPoint));
+			FString VirtualAssetPath = PackagePath.RightChop(MountPoint.Len() + 1);
+			
+			// Create the full path to the asset and assign the bookmark as a child
+			FFolderTreeItemPtr FolderTreeItem = MountPointFolder->CreatePath(VirtualAssetPath);
+			FolderTreeItem->AddChild(WorldBookmarkTreeItem);
+		}
 	}
 
-	SortItems();	
+	// Now that we have repopulated the tree, do not leave empty folders behind
+	TreeRoot->ClearEmptyFolders();
+	
+	SortItems();
+
+	if (bExpandAllItemsOnNextRefresh || bForceRefresh)
+	{
+		SetItemExpansionRecursive(TreeRoot, true);
+	}
 
 	// Restore selected bookmark after refresh
 	SetSelectedBookmark(PreviousSelection);
 
 	bPendingRefresh = false;
+	bExpandAllItemsOnNextRefresh = false;
 }
 
 void SWorldBookmarkBrowser::SortItems()
 {
-	TreeItems.Sort([this](const FTreeItemPtr& LHS, const FTreeItemPtr& RHS)
+	auto CompareFolderTreeItems = [this](const FFolderTreeItem* A, const FFolderTreeItem* B) -> bool
 	{
-		const FWorldBookmarkTreeItemPtr& A = StaticCastSharedPtr<FWorldBookmarkTreeItem>(LHS);
-		const FWorldBookmarkTreeItemPtr& B = StaticCastSharedPtr<FWorldBookmarkTreeItem>(RHS);
+		const bool bReverseSort = (ColumnsSortingParams[EColumnSortPriority::Primary].Key == WorldBookmarkBrowserColumn::Label.Id && ColumnsSortingParams[EColumnSortPriority::Primary].Value == EColumnSortMode::Descending) ||
+							(ColumnsSortingParams[EColumnSortPriority::Secondary].Key == WorldBookmarkBrowserColumn::Label.Id && ColumnsSortingParams[EColumnSortPriority::Secondary].Value == EColumnSortMode::Descending);
+
+		return bReverseSort ? B->GetName().LexicalLess(A->GetName()) : A->GetName().LexicalLess(B->GetName());
+	};
 
+	auto CompareWorldBookmarkTreeItems = [this](const FWorldBookmarkTreeItem* A, const FWorldBookmarkTreeItem* B) -> bool
+	{
 		auto CompareItems = [&](const TPair<FName, EColumnSortMode::Type>& SortOrder) -> int32
 		{
-			auto ApplySortMode = [&](int32 SortResult) { return SortOrder.Value == EColumnSortMode::Ascending ? SortResult : -SortResult; };
-			
-			if (!A.IsValid() || !B.IsValid())
+			auto ApplySortMode = [&](int32 SortResult)
 			{
-				return A.IsValid() ? -1 : 0;
-			}
-
+				return SortOrder.Value == EColumnSortMode::Ascending ? SortResult : -SortResult;
+			};
+			
 			UWorldBookmark* WorldBookmarkA = Cast<UWorldBookmark>(A->BookmarkAsset.GetAsset());
 			UWorldBookmark* WorldBookmarkB = Cast<UWorldBookmark>(B->BookmarkAsset.GetAsset());
 			
@@ -581,6 +725,10 @@ void SWorldBookmarkBrowser::SortItems()
 			{
 				return ApplySortMode(A->GetBookmarkWorld().GetAssetName() == B->GetBookmarkWorld().GetAssetName() ? 0 : A->GetBookmarkWorld().GetAssetName() < B->GetBookmarkWorld().GetAssetName() ? -1 : 1);
 			}
+			else if (SortOrder.Key == WorldBookmarkBrowserColumn::Category.Id)
+			{
+				return ApplySortMode(WorldBookmarkA->GetBookmarkCategory().Name == WorldBookmarkB->GetBookmarkCategory().Name ? 0 : WorldBookmarkA->GetBookmarkCategory() < WorldBookmarkB->GetBookmarkCategory() ? -1 : 1);
+			}
 		
 			checkf(false, TEXT("Unsupported sort order: %s"), *SortOrder.Key.ToString());
 			return 0;
@@ -596,6 +744,30 @@ void SWorldBookmarkBrowser::SortItems()
 		// Secondary Sort
 		int32 SecondarySortOrder = CompareItems(ColumnsSortingParams[EColumnSortPriority::Secondary]);
 		return SecondarySortOrder < 0;
+	};
+
+	TreeRoot->Sort([this, &CompareWorldBookmarkTreeItems, &CompareFolderTreeItems](const FTreeItemPtr& LHS, const FTreeItemPtr& RHS)
+	{
+		if (!LHS.IsValid() || !RHS.IsValid())
+		{
+			return LHS.IsValid();
+		}
+
+		if (LHS->IsA<FFolderTreeItem>() != RHS->IsA<FFolderTreeItem>())
+		{
+			// Folders should always come first
+			return LHS->IsA<FFolderTreeItem>();
+		}
+		else if (LHS->IsA<FFolderTreeItem>())
+		{
+			// Sorting two folders
+			return CompareFolderTreeItems(LHS->Cast<FFolderTreeItem>(), RHS->Cast<FFolderTreeItem>());
+		}
+		else
+		{
+			// Sorting two bookmarks
+			return CompareWorldBookmarkTreeItems(LHS->Cast<FWorldBookmarkTreeItem>(), RHS->Cast<FWorldBookmarkTreeItem>());
+		}
 	});
 
 	if (BookmarksView)
@@ -616,6 +788,10 @@ const FSoftObjectPath& FWorldBookmarkTreeItem::GetBookmarkWorld() const
 
 void SWorldBookmarkBrowser::OnGetChildren(FTreeItemPtr InItem, TArray<FTreeItemPtr>& OutChildren)
 {
+	if (FFolderTreeItem* FolderTreeItem = InItem->Cast<FFolderTreeItem>())
+	{
+		OutChildren = FolderTreeItem->GetChildren();
+	}
 }
 
 TSharedRef<ITableRow> SWorldBookmarkBrowser::OnGenerateRow(FTreeItemPtr InTreeItem, const TSharedRef<STableViewBase>& OwnerTable)
@@ -624,6 +800,10 @@ TSharedRef<ITableRow> SWorldBookmarkBrowser::OnGenerateRow(FTreeItemPtr InTreeIt
 	{
 		return SNew(FWorldBookmarkTableRow, OwnerTable, StaticCastSharedPtr<FWorldBookmarkTreeItem>(InTreeItem).ToSharedRef());
 	}
+	else if (InTreeItem->IsA<FFolderTreeItem>())
+	{
+		return SNew(FFolderTableRow, OwnerTable, StaticCastSharedPtr<FFolderTreeItem>(InTreeItem).ToSharedRef());
+	}
 
 	return SNew(STableRow<TSharedPtr<FString>>, OwnerTable);
 }
@@ -673,6 +853,7 @@ void SWorldBookmarkBrowser::OnSearchBoxTextChanged(const FText& InSearchText)
 	if (!CurrentSearchString.EqualTo(InSearchText))
 	{
 		CurrentSearchString = InSearchText;
+		bExpandAllItemsOnNextRefresh = true;
 		RefreshItems();
 	}
 }
@@ -719,21 +900,38 @@ TSharedPtr<SWidget> SWorldBookmarkBrowser::OnContextMenuOpening()
 {
 	FMenuBuilder MenuBuilder(true, Commands);
 
-	// World Bookmark
-	MenuBuilder.BeginSection(TEXT("WorldBookmarkMenuSection"), LOCTEXT("WorldBookmarkMenuSection_Label", "World Bookmark"));
+	FTreeItemPtr SelectedItem = GetSelectedTreeItem();
+	FWorldBookmarkTreeItem* BookmarkTreeItem = SelectedItem.IsValid() ? SelectedItem->Cast<FWorldBookmarkTreeItem>() : nullptr;
+	FFolderTreeItem* FolderTreeItem = SelectedItem.IsValid() ? SelectedItem->Cast<FFolderTreeItem>() : nullptr;
+
+	// No context menu for virtual folders
+	if (FolderTreeItem && FolderTreeItem->IsVirtual())
 	{
+		return SNullWidget::NullWidget;
+	}
+
+	// World Bookmark
+	if (BookmarkTreeItem)
+	{		
+		MenuBuilder.BeginSection(TEXT("WorldBookmarkMenuSection"), LOCTEXT("WorldBookmarkMenuSection_Label", "World Bookmark"));
 		MenuBuilder.AddMenuEntry(FWorldBookmarkCommands::Get().LoadBookmark);
 		MenuBuilder.AddMenuEntry(FWorldBookmarkCommands::Get().UpdateBookmark);
 		MenuBuilder.AddMenuEntry(FWorldBookmarkCommands::Get().AddToFavorite);
 		MenuBuilder.AddMenuEntry(FWorldBookmarkCommands::Get().RemoveFromFavorite);
 		MenuBuilder.AddMenuEntry(FWorldBookmarkCommands::Get().PlayFromLocation);
 		MenuBuilder.AddMenuEntry(FWorldBookmarkCommands::Get().MoveCameraToLocation);
+		MenuBuilder.EndSection();
 	}
 
-	MenuBuilder.EndSection();
-
-	// Asset Options
-	MenuBuilder.BeginSection(TEXT("AssetOptionsSection"), LOCTEXT("AssetOptionsText", "Asset Options"));
+	// Asset/Folder Options
+	if (FolderTreeItem)
+	{
+		MenuBuilder.BeginSection(TEXT("FolderOptionsSection"), LOCTEXT("FolderOptionsText", "Folder Options"));
+	}
+	else
+	{
+		MenuBuilder.BeginSection(TEXT("AssetOptionsSection"), LOCTEXT("AssetOptionsText", "Asset Options"));
+	}
 	{
 		// Find In Content Browser
 		MenuBuilder.AddMenuEntry(FGlobalEditorCommonCommands::Get().FindInContentBrowser);
@@ -743,6 +941,15 @@ TSharedPtr<SWidget> SWorldBookmarkBrowser::OnContextMenuOpening()
 
 		// Delete
 		MenuBuilder.AddMenuEntry(FGenericCommands::Get().Delete);
+
+		if (FolderTreeItem)
+		{
+			MenuBuilder.AddMenuEntry(FWorldBookmarkCommands::Get().CreateBookmarkInFolder);
+		}
+		else
+		{
+			MenuBuilder.AddMenuEntry(FWorldBookmarkCommands::Get().MoveBookmarkToNewFolder);
+		}
 	}
 	MenuBuilder.EndSection();
 
@@ -848,6 +1055,7 @@ void SWorldBookmarkBrowser::OnSettingsChanged()
 		OnUncontrolledChangelistModuleChangedHandle.Reset();
 	}
 
+	RecreateColumns();
 	RefreshItems();
 }
 
@@ -885,12 +1093,11 @@ void SWorldBookmarkBrowser::UpdateSelectedBookmark()
 	}
 }
 
-void SWorldBookmarkBrowser::FindSelectedBookmarkInContentBrowser()
+void SWorldBookmarkBrowser::FindSelectedItemInContentBrowser()
 {
-	if (UWorldBookmark* SelectedBookmark = GetSelectedBookmark())
+	if (FTreeItemPtr SelectedTreeItem = GetSelectedTreeItem())
 	{
-		FContentBrowserModule& ContentBrowserModule = FModuleManager::Get().LoadModuleChecked<FContentBrowserModule>("ContentBrowser");
-		ContentBrowserModule.Get().SyncBrowserToAssets(TArray<UObject*>({ SelectedBookmark }));
+		SelectedTreeItem->ShowInContentBrowser();
 	}
 }
 
@@ -953,6 +1160,11 @@ FTreeItemPtr SWorldBookmarkBrowser::GetSelectedTreeItem() const
 	return Selection[0];
 }
 
+bool SWorldBookmarkBrowser::IsValidItemSelected() const
+{
+	return GetSelectedTreeItem() != nullptr;
+}
+
 FWorldBookmarkTreeItemPtr SWorldBookmarkBrowser::GetSelectedBookmarkTreeItem() const
 {
 	FTreeItemPtr SelectedTreeItem = GetSelectedTreeItem();
@@ -990,15 +1202,11 @@ void SWorldBookmarkBrowser::SetSelectedBookmark(UWorldBookmark* InBookmark)
 
 	if (InBookmark)
 	{
-		FTreeItemPtr* BookmarkTreeItem = TreeItems.FindByPredicate([InBookmark](FTreeItemPtr InItem)
-		{
-			return InItem->IsA<FWorldBookmarkTreeItem>() && InItem->Cast<FWorldBookmarkTreeItem>()->BookmarkAsset == InBookmark;
-		});
-
+		FWorldBookmarkTreeItemPtr* BookmarkTreeItem = BookmarkTreeItems.Find(InBookmark->GetPackage()->GetFName());
 		if (BookmarkTreeItem)
 		{
 			BookmarksView->SetSelection(*BookmarkTreeItem);
-			BookmarksView->RequestScrollIntoView(*BookmarkTreeItem);
+			BookmarksView->RequestNavigateToItem(*BookmarkTreeItem);
 		}
 	}
 	else
@@ -1023,21 +1231,31 @@ void SWorldBookmarkBrowser::OnItemScrolledIntoView(FTreeItemPtr InTreeItem, cons
 	}
 }
 
-void SWorldBookmarkBrowser::RequestRenameSelectedBookmark()
+void SWorldBookmarkBrowser::RequestRenameSelectedItem()
 {
-	if (FWorldBookmarkTreeItemPtr SelectedBookmarkTreeItem = GetSelectedBookmarkTreeItem())
+	if (FTreeItemPtr SelectedTreeItem = GetSelectedTreeItem())
 	{
 		// Before showing the inline editable text block, we must ensure that the tree item is visible.
-		TreeItemPendingRename = SelectedBookmarkTreeItem;
-		BookmarksView->RequestScrollIntoView(SelectedBookmarkTreeItem);
+		TreeItemPendingRename = SelectedTreeItem;
+		BookmarksView->RequestScrollIntoView(SelectedTreeItem);
 	}
 }
 
-void SWorldBookmarkBrowser::DeleteSelectedBookmark()
+bool SWorldBookmarkBrowser::CanRenameOrDeleteSelectedItem() const
 {
-	if (FWorldBookmarkTreeItemPtr SelectedBookmarkTreeItem = GetSelectedBookmarkTreeItem())
+	if (FTreeItemPtr SelectedTreeItem = GetSelectedTreeItem())
 	{
-		ObjectTools::DeleteAssets({ SelectedBookmarkTreeItem->BookmarkAsset });
+		return SelectedTreeItem->CanRename();
+	}
+
+	return false;
+}
+
+void SWorldBookmarkBrowser::DeleteSelectedItem()
+{
+	if (FTreeItemPtr SelectedTreeItem = GetSelectedTreeItem())
+	{
+		SelectedTreeItem->Delete();
 	}
 }
 
@@ -1129,6 +1347,37 @@ void SWorldBookmarkBrowser::Tick(const FGeometry& AllottedGeometry, const double
 	SCompoundWidget::Tick(AllottedGeometry, InCurrentTime, InDeltaTime);
 }
 
+void SWorldBookmarkBrowser::OnItemExpansionChanged(FTreeItemPtr InTreeItem, bool bIsExpanded)
+{
+	if (FFolderTreeItem* FolderItem = InTreeItem->Cast<FFolderTreeItem>())
+	{
+		FolderItem->SetExpanded(bIsExpanded);
+
+		// Expand any children that are also expanded
+		FolderItem->ForEachChildRecursive<FFolderTreeItem>([this](FFolderTreeItemPtr ChildFolderItem)
+		{
+			if (ChildFolderItem->IsExpanded())
+			{
+				BookmarksView->SetItemExpansion(ChildFolderItem, true);
+			}
+		});
+	}
+}
+
+void SWorldBookmarkBrowser::SetItemExpansionRecursive(FTreeItemPtr InTreeItem, bool bIsExpanded)
+{
+	if (FFolderTreeItem* FolderItem = InTreeItem->Cast<FFolderTreeItem>())
+	{
+		BookmarksView->SetItemExpansion(InTreeItem, bIsExpanded);
+
+		// Expand any children
+		FolderItem->ForEachChildRecursive<FFolderTreeItem>([this](FFolderTreeItemPtr ChildFolderItem)
+		{
+			BookmarksView->SetItemExpansion(ChildFolderItem, true);
+		});
+	}
+}
+
 void SWorldBookmarkBrowser::UpdateDetailsView(UWorldBookmark* SelectedWorldBookmark)
 {
 	DetailsBox->ClearChildren();
@@ -1169,5 +1418,74 @@ void SWorldBookmarkBrowser::UpdateDetailsView(UWorldBookmark* SelectedWorldBookm
 	}
 }
 
+void SWorldBookmarkBrowser::MoveSelectedBookmarkToNewFolder()
+{
+	FTreeItemPtr ItemToMove = GetSelectedTreeItem();
+	if (ItemToMove)
+	{
+		FFolderTreeItem* ParentFolder = ItemToMove->GetParent()->Cast<FFolderTreeItem>();
+		if (ParentFolder)
+		{
+			bool bValidName = true;
+			FString BaseFolderName(TEXT("NewFolder"));
+			FString NewFolderName(BaseFolderName);
+			int32 FolderCounter = 0;
+			do
+			{
+				bValidName = true;
+				ParentFolder->ForEachChild<FFolderTreeItem>([NewFolderName, &bValidName](FFolderTreeItemPtr ChildFolder)
+				{
+					bValidName &= ChildFolder->GetName() == NewFolderName;
+				});
+
+				if (!bValidName)
+				{
+					NewFolderName = FString::Printf(TEXT("%s%d"), *BaseFolderName, FolderCounter);
+					FolderCounter++;
+				}
+			} while (!bValidName);
+
+			FFolderTreeItemPtr NewFolder = ParentFolder->CreatePath(NewFolderName);
+
+			NewFolder->Move(ItemToMove);
+			BookmarksView->SetSelection(NewFolder);
+			RequestRenameSelectedItem();
+		}
+	}
+}
+
+void SWorldBookmarkBrowser::CreateBookmarkInSelectedFolder()
+{
+	if (FTreeItemPtr SelectedTreeItem = GetSelectedTreeItem())
+	{
+		if (FFolderTreeItem* SelectedFolder = SelectedTreeItem->Cast<FFolderTreeItem>())
+		{
+			FString AssetPath = SelectedFolder->GetAssetPath();
+			
+			FString AssetName;
+			FString PackageName;
+
+			IAssetTools* AssetTools = &FModuleManager::GetModuleChecked<FAssetToolsModule>("AssetTools").Get();
+			AssetTools->CreateUniqueAssetName(AssetPath / TEXT("NewWorldBookmark"), FString(), PackageName, AssetName);
+			UObject* NewAsset = AssetTools->CreateAsset(AssetName, AssetPath, UWorldBookmark::StaticClass(), nullptr);
+
+			if (UWorldBookmark* NewWorldBookmark = Cast<UWorldBookmark>(NewAsset))
+			{
+				// Show it in the content browser
+				GEditor->SyncBrowserToObject(NewWorldBookmark);
+
+				// Force refresh the bookmark browser, to ensure the new bookmark tree item is created
+				RefreshItems(/*bForceRefresh=*/true);
+
+				// Make sure the new bookmark is selected
+				SetSelectedBookmark(NewWorldBookmark);
+
+				// Enter editing mode on the new bookmark label
+				RequestRenameSelectedItem();
+			}
+		}
+	}
+}
+
 #undef LOCTEXT_NAMESPACE
 
diff --git a/Engine/Source/Editor/WorldBookmark/Private/WorldBookmark/Browser/SWorldBookmarkBrowser.h b/Engine/Source/Editor/WorldBookmark/Private/WorldBookmark/Browser/SWorldBookmarkBrowser.h
--- a/Engine/Source/Editor/WorldBookmark/Private/WorldBookmark/Browser/SWorldBookmarkBrowser.h
+++ b/Engine/Source/Editor/WorldBookmark/Private/WorldBookmark/Browser/SWorldBookmarkBrowser.h
@@ -8,6 +8,7 @@
 #include "UnrealEdMisc.h"
 #include "Widgets/Views/STreeView.h"
 
+#include "WorldBookmark/Browser/FolderTreeItem.h"
 #include "WorldBookmark/Browser/BookmarkTreeItem.h"
 
 class UWorldBookmark;
@@ -48,6 +49,7 @@ private:
 	void RefreshItems(bool bForceRefresh = false);
 	void SortItems();
 
+	void RecreateColumns();
 	EColumnSortMode::Type GetColumnSortMode(const FName ColumnId) const;
 	EColumnSortPriority::Type GetColumnSortPriority(const FName InColumnId) const;
 	TArray<FAssetViewCustomColumn> GetCustomColumns() const;
@@ -70,21 +72,24 @@ private:
 	FWorldBookmarkTreeItemPtr GetSelectedBookmarkTreeItem() const;
 	UWorldBookmark* GetSelectedBookmark() const;
 	void SetSelectedBookmark(UWorldBookmark* InBookmark);
+	bool IsValidBookmarkSelected() const;
+	bool IsValidItemSelected() const;
 
 	bool CanExecuteBookmarkAction() const;
 
 	void LoadSelectedBookmark();
 	void UpdateSelectedBookmark();
-	void FindSelectedBookmarkInContentBrowser();
-	bool IsValidBookmarkSelected() const;
 
 	bool IsSelectedBookmarkFavorite() const;
 	bool IsSelectedBookmarkNotFavorite() const;
 	void AddSelectedBookmarkToFavorites();
 	void RemoveSelectedBookmarkFromFavorites();	
 
-	void RequestRenameSelectedBookmark();
-	void DeleteSelectedBookmark();
+	void FindSelectedItemInContentBrowser();
+	void RequestRenameSelectedItem();
+
+	bool CanRenameOrDeleteSelectedItem() const;
+	void DeleteSelectedItem();
 
 	bool CanGoToSelectedBookmarkLocation() const;
 	void PlayFromSelectedBookmarkLocation();
@@ -92,6 +97,12 @@ private:
 
 	void UpdateDetailsView(UWorldBookmark* InSelectedWorldBookmark);
 
+	void OnItemExpansionChanged(FTreeItemPtr InTreeItem, bool bIsExpanded);
+	void SetItemExpansionRecursive(FTreeItemPtr InTreeItem, bool bIsExpanded);
+
+	void MoveSelectedBookmarkToNewFolder();
+	void CreateBookmarkInSelectedFolder();
+
 private:
 	FDelegateHandle OnAssetAddedHandle;
 	FDelegateHandle OnAssetRemovedHandle;
@@ -103,8 +114,10 @@ private:
 	FDelegateHandle OnSettingsChangedHandle;
 	FDelegateHandle OnUncontrolledChangelistModuleChangedHandle;
 
-	TArray<FTreeItemPtr> TreeItems;
+	FFolderTreeItemPtr TreeRoot;
+	TMap<FName, FWorldBookmarkTreeItemPtr> BookmarkTreeItems;
 
+	TSharedPtr<SHeaderRow> HeaderRow;
 	TSharedPtr<STreeView<FTreeItemPtr>> BookmarksView;
 	TSharedPtr<IDetailsView> PropertyView;
 	TSharedPtr<SVerticalBox> DetailsBox;
@@ -119,4 +132,5 @@ private:
 
 	bool bIsInTick;
 	bool bPendingRefresh;
+	bool bExpandAllItemsOnNextRefresh;
 };
\ No newline at end of file
diff --git a/Engine/Source/Editor/WorldBookmark/Private/WorldBookmark/Browser/Settings.cpp b/Engine/Source/Editor/WorldBookmark/Private/WorldBookmark/Browser/Settings.cpp
--- a/Engine/Source/Editor/WorldBookmark/Private/WorldBookmark/Browser/Settings.cpp
+++ b/Engine/Source/Editor/WorldBookmark/Private/WorldBookmark/Browser/Settings.cpp
@@ -109,4 +109,24 @@ void UWorldBookmarkBrowserSettings::SetMaxLastRecentlyUsedItems(int32 InMaxLastR
 		Settings.MaxLastRecentlyUsedItems = InMaxLastRecentlyUsedItems;
 		Settings.ApplySettingsChanges();
 	}
+}
+
+EWorldBookmarkBrowserViewMode UWorldBookmarkBrowserSettings::GetViewMode()
+{
+	return Get().ViewMode;
+}
+
+bool UWorldBookmarkBrowserSettings::IsViewMode(EWorldBookmarkBrowserViewMode InViewMode)
+{
+	return GetViewMode() == InViewMode;
+}
+
+void UWorldBookmarkBrowserSettings::SetViewMode(EWorldBookmarkBrowserViewMode InViewMode)
+{
+	UWorldBookmarkBrowserSettings& Settings = Get();
+	if (Settings.ViewMode != InViewMode)
+	{
+		Settings.ViewMode = InViewMode;
+		Settings.ApplySettingsChanges();
+	}
 }
\ No newline at end of file
diff --git a/Engine/Source/Editor/WorldBookmark/Private/WorldBookmark/Browser/Settings.h b/Engine/Source/Editor/WorldBookmark/Private/WorldBookmark/Browser/Settings.h
--- a/Engine/Source/Editor/WorldBookmark/Private/WorldBookmark/Browser/Settings.h
+++ b/Engine/Source/Editor/WorldBookmark/Private/WorldBookmark/Browser/Settings.h
@@ -5,6 +5,13 @@
 
 #include "Settings.generated.h"
 
+UENUM()
+enum class EWorldBookmarkBrowserViewMode : uint8
+{
+	ListView,
+	TreeView
+};
+
 UCLASS(EditorConfig = "WorldBookmarkBrowser")
 class UWorldBookmarkBrowserSettings : public UEditorConfigBase
 {
@@ -33,6 +40,10 @@ public:
 	static int32 GetMaxLastRecentlyUsedItems();
 	static void SetMaxLastRecentlyUsedItems(int32 InMaxLastRecentlyUsedItems);
 
+	static EWorldBookmarkBrowserViewMode GetViewMode();
+	static bool IsViewMode(EWorldBookmarkBrowserViewMode InViewMode);
+	static void SetViewMode(EWorldBookmarkBrowserViewMode InViewMode);
+
 private:
 	static UWorldBookmarkBrowserSettings& Get();
 
@@ -56,6 +67,9 @@ private:
 	UPROPERTY(meta=(EditorConfig))
 	TArray<FName> HiddenColumns;
 
+	UPROPERTY(meta=(EditorConfig))
+	EWorldBookmarkBrowserViewMode ViewMode;
+
 	FWorldBookmarkBrowserSettingChanged OnSettingsChangedDelegate;
 
 	static TObjectPtr<UWorldBookmarkBrowserSettings> Instance;
diff --git a/Engine/Source/Editor/WorldBookmark/Private/WorldBookmark/Browser/TreeItem.cpp b/Engine/Source/Editor/WorldBookmark/Private/WorldBookmark/Browser/TreeItem.cpp
--- a/Engine/Source/Editor/WorldBookmark/Private/WorldBookmark/Browser/TreeItem.cpp
+++ b/Engine/Source/Editor/WorldBookmark/Private/WorldBookmark/Browser/TreeItem.cpp
@@ -1,3 +1,40 @@
 // Copyright Epic Games, Inc. All Rights Reserved.
 
 #include "WorldBookmark/Browser/TreeItem.h"
+
+ITreeItem::ITreeItem(ETreeItemType InType)
+	: TreeItemType(InType)
+	, ParentItem(nullptr)
+{
+}
+
+ITreeItem::~ITreeItem()
+{
+}
+
+ITreeItem* ITreeItem::GetParent() const
+{
+	return ParentItem;
+}
+
+void ITreeItem::SetParent(ITreeItem* InParent)
+{
+	ParentItem = InParent;
+}
+
+FString ITreeItem::GetAssetPath() const
+{
+	TStringBuilder<256> FullPathBuilder;
+
+	const ITreeItem* CurrentItem = this;
+	do
+	{
+		if (!CurrentItem->GetAssetName().IsNone())
+		{
+			FullPathBuilder.Prepend(CurrentItem->GetAssetName().ToString());
+			FullPathBuilder.Prepend(TEXT("/"));
+		}
+	} while (CurrentItem = CurrentItem->GetParent(), CurrentItem);
+
+	return FString(FullPathBuilder);
+}
\ No newline at end of file
diff --git a/Engine/Source/Editor/WorldBookmark/Private/WorldBookmark/Browser/TreeItem.h b/Engine/Source/Editor/WorldBookmark/Private/WorldBookmark/Browser/TreeItem.h
--- a/Engine/Source/Editor/WorldBookmark/Private/WorldBookmark/Browser/TreeItem.h
+++ b/Engine/Source/Editor/WorldBookmark/Private/WorldBookmark/Browser/TreeItem.h
@@ -14,20 +14,18 @@ enum ETreeItemType
 struct ITreeItem : public TSharedFromThis<ITreeItem>
 {
 public:
-	virtual FString GetName() const = 0;
-
 	/** Attempt to cast this item to another type if it is of that type. Returns nullptr if it fails */
 	template <typename T>
 	T* Cast()
 	{
-		return TreeItemType == T::Type ? StaticCast<T*>(this) : nullptr;
+		return IsA<T>() ? StaticCast<T*>(this) : nullptr;
 	}
 
 	/** Attempt to cast this item to another type if it is of that type. Returns nullptr if it fails */
 	template <typename T>
 	const T* Cast() const
 	{
-		return TreeItemType == T::Type ? StaticCast<const T*>(this) : nullptr;
+		return IsA<T>() ? StaticCast<const T*>(this) : nullptr;
 	}
 
 	/** Returns true if this item is of the specified type */
@@ -36,6 +34,32 @@ public:
 	{
 		return TreeItemType == T::Type;
 	}
+
+	template <>
+	bool IsA<ITreeItem>() const
+	{
+		return true;
+	}
+
+	virtual FName GetName() const = 0;
+
+	virtual FName GetAssetName() const = 0;
+
+	virtual bool CanRename() const = 0;
+
+	virtual bool TryRename(FName InNewName, FText& OutErrorMessage) const = 0;
+	
+	virtual bool Rename(FName InNewName) = 0;
+
+	virtual bool Delete() = 0;
+
+	virtual FName GetIconName() const = 0;
+
+	virtual FText GetText() const = 0;
+
+	virtual FText GetTooltipText() const = 0;
+
+	virtual void ShowInContentBrowser() const = 0;
 	
 	bool bInEditingMode = false;
 
@@ -46,12 +70,18 @@ public:
 	/** Tree item type identifier */
 	ETreeItemType TreeItemType;
 
+	ITreeItem* GetParent() const;
+
+	void SetParent(ITreeItem* InParent);
+
+	FString GetAssetPath() const;
+
+private:
+	ITreeItem* ParentItem;
+
 protected:
-	/** Default constructor */
-	ITreeItem(ETreeItemType InType) : TreeItemType(InType) {}
-	virtual ~ITreeItem()
-	{
-	}
+	ITreeItem(ETreeItemType InType);
+	virtual ~ITreeItem();
 };
 
 typedef TSharedPtr<ITreeItem> FTreeItemPtr;
diff --git a/Engine/Source/Editor/WorldBookmark/Private/WorldBookmark/WorldBookmark.cpp b/Engine/Source/Editor/WorldBookmark/Private/WorldBookmark/WorldBookmark.cpp
--- a/Engine/Source/Editor/WorldBookmark/Private/WorldBookmark/WorldBookmark.cpp
+++ b/Engine/Source/Editor/WorldBookmark/Private/WorldBookmark/WorldBookmark.cpp
@@ -74,6 +74,12 @@ FName UWorldBookmark::GetWorldNameAssetTag()
 	return NAME_WorldName;
 }
 
+FName UWorldBookmark::GetCategoryAssetTag()
+{
+	static const FName NAME_WorldBookmarkCategory = "WorldBookmarkCategory";
+	return NAME_WorldBookmarkCategory;
+}
+
 FSoftObjectPath UWorldBookmark::GetWorldFromAssetData(const FAssetData& InAssetData)
 {
 	FSoftObjectPath WorldPath;
@@ -88,6 +94,19 @@ FSoftObjectPath UWorldBookmark::GetWorldFromAssetData(const FAssetData& InAssetD
 	return WorldPath;
 }
 
+FWorldBookmarkCategory UWorldBookmark::GetCategoryFromAssetData(const FAssetData& InAssetData)
+{
+	FWorldBookmarkCategory Category;
+
+	FString CategoryImportText;
+	if (InAssetData.GetTagValue(GetCategoryAssetTag(), CategoryImportText))
+	{
+		FWorldBookmarkCategory::StaticStruct()->ImportText(*CategoryImportText, &Category, nullptr, PPF_None, GLog, FWorldBookmarkCategory::StaticStruct()->GetName());
+	}
+
+	return Category;
+}
+
 void UWorldBookmark::GetAssetRegistryTags(FAssetRegistryTagsContext Context) const
 {
 	Super::GetAssetRegistryTags(Context);
@@ -102,6 +121,13 @@ void UWorldBookmark::GetAssetRegistryTags(FAssetRegistryTagsContext Context) con
 
 			Context.AddTag(FAssetRegistryTag(GetWorldNameAssetTag(), WorldPath.ToString(), FAssetRegistryTag::TT_Hidden));
 		}
+
+		if (!Category.Name.IsNone())
+		{
+			FString CategoryExportText;
+			FWorldBookmarkCategory::StaticStruct()->ExportText(CategoryExportText, &Category, nullptr, nullptr, PPF_None, nullptr);
+			Context.AddTag(FAssetRegistryTag(GetCategoryAssetTag(), CategoryExportText, FAssetRegistryTag::TT_Hidden));
+		}
 	}
 }
 
@@ -283,4 +309,19 @@ void UWorldBookmark::SetUserLastLoadedTimeStampUTC(const FDateTime& InLastLoaded
 	SaveConfig();
 }
 
+const FWorldBookmarkCategory& UWorldBookmark::GetBookmarkCategory() const
+{
+	return Category;
+}
+
+bool FWorldBookmarkCategory::operator<(const FWorldBookmarkCategory& Other) const
+{
+	if (Name.IsNone() || Other.Name.IsNone())
+	{
+		return Name.IsNone();
+	}
+
+	return Name.LexicalLess(Other.Name);
+}
+
 #undef LOCTEXT_NAMESPACE
diff --git a/Engine/Source/Editor/WorldBookmark/Private/WorldBookmark/WorldBookmarkCommands.cpp b/Engine/Source/Editor/WorldBookmark/Private/WorldBookmark/WorldBookmarkCommands.cpp
--- a/Engine/Source/Editor/WorldBookmark/Private/WorldBookmark/WorldBookmarkCommands.cpp
+++ b/Engine/Source/Editor/WorldBookmark/Private/WorldBookmark/WorldBookmarkCommands.cpp
@@ -22,6 +22,9 @@ void FWorldBookmarkCommands::RegisterCommands()
 	UI_COMMAND(RemoveFromFavorite, "Remove from Favorite", "Remove this bookmark from your favorites.", EUserInterfaceActionType::Button, FInputChord());
 	UI_COMMAND(PlayFromLocation, "Play From Here", "Play from here.", EUserInterfaceActionType::Button, FInputChord());
 	UI_COMMAND(MoveCameraToLocation, "Move Camera Here", "Move the camera to the selected location.", EUserInterfaceActionType::Button, FInputChord());
+
+	UI_COMMAND(MoveBookmarkToNewFolder, "Move to New Folder", "Create a new folder containing the current selection.", EUserInterfaceActionType::Button, FInputChord());
+	UI_COMMAND(CreateBookmarkInFolder, "Create Bookmark", "Create a new bookmark in the selected folder.", EUserInterfaceActionType::Button, FInputChord());
 }
 
 #undef LOCTEXT_NAMESPACE
diff --git a/Engine/Source/Editor/WorldBookmark/Private/WorldBookmark/WorldBookmarkCommands.h b/Engine/Source/Editor/WorldBookmark/Private/WorldBookmark/WorldBookmarkCommands.h
--- a/Engine/Source/Editor/WorldBookmark/Private/WorldBookmark/WorldBookmarkCommands.h
+++ b/Engine/Source/Editor/WorldBookmark/Private/WorldBookmark/WorldBookmarkCommands.h
@@ -15,6 +15,8 @@ public:
 	TSharedPtr<FUICommandInfo> RemoveFromFavorite;
 	TSharedPtr<FUICommandInfo> PlayFromLocation;
 	TSharedPtr<FUICommandInfo> MoveCameraToLocation;
+	TSharedPtr<FUICommandInfo> MoveBookmarkToNewFolder;
+	TSharedPtr<FUICommandInfo> CreateBookmarkInFolder;
 	
 	/**
 	 * Initialize commands
diff --git a/Engine/Source/Editor/WorldBookmark/Private/WorldBookmark/WorldBookmarkDetailsCustomization.cpp b/Engine/Source/Editor/WorldBookmark/Private/WorldBookmark/WorldBookmarkDetailsCustomization.cpp
--- a/Engine/Source/Editor/WorldBookmark/Private/WorldBookmark/WorldBookmarkDetailsCustomization.cpp
+++ b/Engine/Source/Editor/WorldBookmark/Private/WorldBookmark/WorldBookmarkDetailsCustomization.cpp
@@ -2,10 +2,24 @@
 
 #include "WorldBookmark/WorldBookmarkDetailsCustomization.h"
 #include "WorldBookmark/WorldBookmark.h"
+
+#include "AssetRegistry/AssetData.h"
+#include "AssetRegistry/IAssetRegistry.h"
 #include "DetailLayoutBuilder.h"
 #include "DetailCategoryBuilder.h"
+#include "DetailWidgetRow.h"
 #include "Editor.h"
 #include "EditorState/WorldEditorState.h"
+#include "IDetailChildrenBuilder.h"
+#include "IStructureDetailsView.h"
+#include "PropertyHandle.h"
+#include "SPrimaryButton.h"
+#include "Widgets/Colors/SColorBlock.h"
+#include "Widgets/Input/SComboBox.h"
+#include "Widgets/Layout/SUniformGridPanel.h"
+#include "Widgets/Text/STextBlock.h"
+
+#define LOCTEXT_NAMESPACE "WorldBookmarkDetailsCustomization"
 
 FWorldBookmarkDetailsCustomization::FWorldBookmarkDetailsCustomization()
 	: CachedDetailBuilder()
@@ -73,3 +87,305 @@ void FWorldBookmarkDetailsCustomization::PostRedo(bool bSuccess)
 	// Refresh the UI
 	RefreshCustomDetail();
 }
+
+TSharedRef<IPropertyTypeCustomization> FWorldBookmarkCategoryCustomization::MakeInstance()
+{
+	return MakeShareable(new FWorldBookmarkCategoryCustomization);
+}
+
+
+const FName FWorldBookmarkCategoryCustomization::NAME_NewCategory("New Category...");
+
+FWorldBookmarkCategoryCustomization::FWorldBookmarkCategoryCustomization()
+{	
+	IAssetRegistry& AssetRegistry = IAssetRegistry::GetChecked();
+	TArray<FAssetData> AssetsData;
+	const bool bSearchSubClasses = true;
+	AssetRegistry.GetAssetsByClass(UWorldBookmark::StaticClass()->GetClassPathName(), AssetsData, bSearchSubClasses);
+		
+	// Add default "categories" to the list of items that will be shown in the Category combo box
+	AddCategory(FWorldBookmarkCategory{ NAME_None, FColor::Black });
+	AddCategory(FWorldBookmarkCategory{ NAME_NewCategory, FColor::Black });
+
+	for (const FAssetData& AssetData : AssetsData)
+	{
+		FWorldBookmarkCategory Category = UWorldBookmark::GetCategoryFromAssetData(AssetData);
+		AddCategory(Category);
+	}
+
+	SortCategories();
+}
+
+void FWorldBookmarkCategoryCustomization::AddCategory(const FWorldBookmarkCategory& InCategory)
+{
+	bool bFoundExistingCategory = KnownCategories.ContainsByPredicate([&InCategory](const TSharedPtr<FWorldBookmarkCategory>& ExistingCategory)
+	{
+		return ExistingCategory->Name == InCategory.Name;
+	});
+
+	if (!bFoundExistingCategory)
+	{
+		KnownCategories.Emplace(MakeShared<FWorldBookmarkCategory>(InCategory));
+	}
+}
+
+void FWorldBookmarkCategoryCustomization::SortCategories()
+{
+	KnownCategories.Sort([](const TSharedPtr<FWorldBookmarkCategory>& LHS, const TSharedPtr<FWorldBookmarkCategory>& RHS)
+	{
+		if (LHS->Name == NAME_None || RHS->Name == NAME_None)
+		{
+			return LHS->Name == NAME_None;
+		}
+
+		if (LHS->Name == NAME_NewCategory || RHS->Name == NAME_NewCategory)
+		{
+			return LHS->Name == NAME_NewCategory;
+		}
+
+		return LHS->Name.LexicalLess(RHS->Name);
+	});
+}
+
+void FWorldBookmarkCategoryCustomization::CustomizeHeader(TSharedRef<IPropertyHandle> InStructPropertyHandle, FDetailWidgetRow& InHeaderRow, IPropertyTypeCustomizationUtils& InCustomizationUtils)
+{
+	StructPropertyHandle = InStructPropertyHandle.ToSharedPtr();
+
+	TArray<void*> RawData;
+	InStructPropertyHandle->AccessRawData(RawData);
+	FWorldBookmarkCategory* CurrentCategory = reinterpret_cast<FWorldBookmarkCategory*>(RawData[0]);
+
+	TSharedPtr<FWorldBookmarkCategory>* CurrentCategoryPtr = KnownCategories.FindByPredicate([CurrentCategory](const TSharedPtr<FWorldBookmarkCategory>& Item)
+	{
+		return Item->Name == CurrentCategory->Name;
+	});
+
+	// Create the collapsed view with Name and Color
+	InHeaderRow.NameContent()
+		[
+			StructPropertyHandle->CreatePropertyNameWidget()
+		]
+		.ValueContent()
+		.MinDesiredWidth(200.0f)
+		[
+			SNew(SComboBox<TSharedPtr<FWorldBookmarkCategory>>)
+				.OptionsSource(&KnownCategories)
+				.OnGenerateWidget(this, &FWorldBookmarkCategoryCustomization::MakeCategoryComboWidget)
+				.OnSelectionChanged(this, &FWorldBookmarkCategoryCustomization::OnCategoryChanged)
+				.InitiallySelectedItem(CurrentCategoryPtr ? *CurrentCategoryPtr : TSharedPtr<FWorldBookmarkCategory>())
+				[
+					SNew(SHorizontalBox)
+
+					+ SHorizontalBox::Slot()
+					.MaxWidth(20)
+					[
+						SNew(SBox)
+						.Visibility(this, &FWorldBookmarkCategoryCustomization::GetSelectedCategoryColorVisibility)
+						.Padding(0, 0, 4, 0)						
+						[
+							SNew(SColorBlock)
+							.Color(this, &FWorldBookmarkCategoryCustomization::GetSelectedCategoryColor)
+							.CornerRadius(FVector4(4.0f, 4.0f, 4.0f, 4.0f))
+							.Size(FVector2D(20, 16))
+						]
+					]
+
+					+ SHorizontalBox::Slot()
+					.VAlign(VAlign_Center)
+					.AutoWidth()
+					[
+						SNew(STextBlock)
+						.Text(this, &FWorldBookmarkCategoryCustomization::GetSelectedCategoryText)
+					]
+				]
+		];
+}
+
+FLinearColor FWorldBookmarkCategoryCustomization::GetSelectedCategoryColor() const
+{
+	TArray<void*> RawData;
+	StructPropertyHandle->AccessRawData(RawData);
+
+	bool bShouldAssign = true;
+	FWorldBookmarkCategory* WorldBookmarkCategory = reinterpret_cast<FWorldBookmarkCategory*>(RawData[0]);
+	return WorldBookmarkCategory->Color;
+}
+
+EVisibility FWorldBookmarkCategoryCustomization::GetSelectedCategoryColorVisibility() const
+{
+	TArray<void*> RawData;
+	StructPropertyHandle->AccessRawData(RawData);
+
+	bool bShouldAssign = true;
+	FWorldBookmarkCategory* WorldBookmarkCategory = reinterpret_cast<FWorldBookmarkCategory*>(RawData[0]);
+
+	const bool bShowColorBlock = (WorldBookmarkCategory->Name != NAME_None) && (WorldBookmarkCategory->Name != NAME_NewCategory);
+	return bShowColorBlock ? EVisibility::Visible : EVisibility::Collapsed;
+}
+
+FText FWorldBookmarkCategoryCustomization::GetSelectedCategoryText() const
+{
+	TArray<void*> RawData;
+	StructPropertyHandle->AccessRawData(RawData);
+
+	bool bShouldAssign = true;
+	FWorldBookmarkCategory* WorldBookmarkCategory = reinterpret_cast<FWorldBookmarkCategory*>(RawData[0]);
+
+	return FText::FromName(WorldBookmarkCategory->Name);
+}
+
+TSharedRef<SWidget> FWorldBookmarkCategoryCustomization::MakeCategoryComboWidget(TSharedPtr<FWorldBookmarkCategory> InItem)
+{
+	TSharedRef<SHorizontalBox> ItemWidget = SNew(SHorizontalBox);
+
+	const bool bShowColorBlock = (InItem->Name != NAME_None) && (InItem->Name != NAME_NewCategory);
+	if (bShowColorBlock)
+	{
+		ItemWidget->AddSlot()
+		.MaxWidth(16)
+		[
+			SNew(SColorBlock)
+				.Color(InItem->Color)
+				.CornerRadius(FVector4(4.0f, 4.0f, 4.0f, 4.0f))
+		];
+	}
+
+	ItemWidget->AddSlot()
+	.VAlign(VAlign_Center)
+	.AutoWidth()
+	.Padding(bShowColorBlock ? 4 : 0, 0, 0, 0)
+	[
+		SNew(STextBlock)
+			.Text(FText::FromName(InItem->Name))
+	];
+
+	return ItemWidget;
+}
+
+void FWorldBookmarkCategoryCustomization::OnCategoryChanged(TSharedPtr<FWorldBookmarkCategory> NewSelection, ESelectInfo::Type SelectInfo)
+{
+	if (NewSelection.IsValid())
+	{
+		TArray<void*> RawData;
+		StructPropertyHandle->AccessRawData(RawData);
+
+		bool bShouldAssign = true;
+		FWorldBookmarkCategory* WorldBookmarkCategory = reinterpret_cast<FWorldBookmarkCategory*>(RawData[0]);
+		FWorldBookmarkCategory CategoryToAssign = *NewSelection;
+
+		if (CategoryToAssign.Name == NAME_NewCategory)
+		{
+			bShouldAssign = CreateNewCategory(CategoryToAssign);
+			if (bShouldAssign)
+			{
+				AddCategory(CategoryToAssign);
+			}
+		}
+
+		if (bShouldAssign)
+		{
+			StructPropertyHandle->NotifyPreChange();
+			*WorldBookmarkCategory = CategoryToAssign;
+			StructPropertyHandle->NotifyPostChange(EPropertyChangeType::ValueSet);
+		}
+	}
+}
+
+bool FWorldBookmarkCategoryCustomization::CreateNewCategory(FWorldBookmarkCategory& OutNewCategory)
+{
+	FPropertyEditorModule& PropertyEditorModule = FModuleManager::GetModuleChecked<FPropertyEditorModule>("PropertyEditor");
+
+	FStructureDetailsViewArgs StructureViewArgs;
+	StructureViewArgs.bShowObjects = true;
+	StructureViewArgs.bShowAssets = true;
+	StructureViewArgs.bShowClasses = true;
+	StructureViewArgs.bShowInterfaces = true;
+
+	FDetailsViewArgs ViewArgs;
+	ViewArgs.bAllowSearch = false;
+	ViewArgs.bHideSelectionTip = false;
+	ViewArgs.bShowObjectLabel = false;
+
+	FWorldBookmarkCategory NewCategoryTemplate;
+	NewCategoryTemplate.Color = FColor::White;
+	TSharedPtr<TStructOnScope<FWorldBookmarkCategory>> Settings;
+	Settings = MakeShared<TStructOnScope<FWorldBookmarkCategory>>();
+	Settings->InitializeAs<FWorldBookmarkCategory>(NewCategoryTemplate);
+
+	TSharedPtr<IStructureDetailsView> DetailsView = PropertyEditorModule.CreateStructureDetailView(ViewArgs, StructureViewArgs, TSharedPtr<FStructOnScope>());
+
+	// Set provided objects on SDetailsView
+	DetailsView->SetStructureData(Settings);
+
+	bool bCreatedNewCategory = false;
+
+	const TSharedRef<SWindow> Window = SNew(SWindow)
+		.Title(LOCTEXT("NewWorldBookmarkCategoryTitle", "New World Bookmark Category"))
+		.SizingRule(ESizingRule::Autosized)
+		.SupportsMaximize(false)
+		.SupportsMinimize(false)
+		.ClientSize(FVector2D(350, 450));
+
+	Window->SetContent(
+		SNew(SVerticalBox)
+
+		+ SVerticalBox::Slot() // Add user input block
+		.Padding(2, 2, 2, 4)
+		[
+			DetailsView->GetWidget().ToSharedRef()
+		]
+
+		+ SVerticalBox::Slot()
+		.AutoHeight()
+		.HAlign(HAlign_Right)
+		.VAlign(VAlign_Bottom)
+		.Padding(8.f, 16.f)
+		[
+			SNew(SUniformGridPanel)
+				.MinDesiredSlotWidth(FAppStyle::GetFloat("StandardDialog.MinDesiredSlotWidth"))
+				.MinDesiredSlotHeight(FAppStyle::GetFloat("StandardDialog.MinDesiredSlotHeight"))
+				.SlotPadding(FAppStyle::GetMargin("StandardDialog.SlotPadding"))
+				+ SUniformGridPanel::Slot(0, 0)
+				[
+					SNew(SPrimaryButton)
+						.Text(LOCTEXT("Ok", "Ok"))
+						.OnClicked_Lambda([&bCreatedNewCategory, Window]
+						{
+							bCreatedNewCategory = true;
+							Window->RequestDestroyWindow();
+							return FReply::Handled();
+						})
+				]
+				+ SUniformGridPanel::Slot(1, 0)
+				[
+					SNew(SButton)
+						.Text(LOCTEXT("Cancel", "Cancel"))
+						.ContentPadding(FAppStyle::GetMargin("StandardDialog.ContentPadding"))
+						.OnClicked_Lambda([&bCreatedNewCategory, Window]
+						{
+							bCreatedNewCategory = false;
+							Window->RequestDestroyWindow();
+							return FReply::Handled();
+						})
+				]
+		]
+	);
+
+	GEditor->EditorAddModalWindow(Window);
+	
+	if (bCreatedNewCategory)
+	{
+		OutNewCategory.Name = Settings->Get()->Name;
+		OutNewCategory.Color = Settings->Get()->Color;
+	}
+
+	bCreatedNewCategory = OutNewCategory.Name != NAME_None && OutNewCategory.Name.ToString().ToLower() != "none";
+	return bCreatedNewCategory;
+}
+
+void FWorldBookmarkCategoryCustomization::CustomizeChildren(TSharedRef<IPropertyHandle> InStructPropertyHandle, IDetailChildrenBuilder& ChildBuilder, IPropertyTypeCustomizationUtils& CustomizationUtils)
+{
+	// Do not allow the row to be expanded (never show the category color & name as individual rows)
+}
+
+#undef LOCTEXT_NAMESPACE
diff --git a/Engine/Source/Editor/WorldBookmark/Private/WorldBookmark/WorldBookmarkDetailsCustomization.h b/Engine/Source/Editor/WorldBookmark/Private/WorldBookmark/WorldBookmarkDetailsCustomization.h
--- a/Engine/Source/Editor/WorldBookmark/Private/WorldBookmark/WorldBookmarkDetailsCustomization.h
+++ b/Engine/Source/Editor/WorldBookmark/Private/WorldBookmark/WorldBookmarkDetailsCustomization.h
@@ -3,9 +3,12 @@
 #pragma once
 
 #include "IDetailCustomization.h"
+#include "IPropertyTypeCustomization.h"
 #include "EditorUndoClient.h"
+#include "Widgets/SWidget.h"
 
 class UWorldBookmark;
+struct FWorldBookmarkCategory;
 
 /** UI customization for UWorldBookmark */
 class FWorldBookmarkDetailsCustomization : public IDetailCustomization, public FEditorUndoClient
@@ -32,3 +35,35 @@ private:
 	/** The detail builder for this cusomtomisation. */
 	IDetailLayoutBuilder* CachedDetailBuilder;
 };
+
+/** UI customization for FWorldBookmarkCategory */
+class FWorldBookmarkCategoryCustomization : public IPropertyTypeCustomization
+{
+public:
+	/** Factory method to create an instance */
+	static TSharedRef<IPropertyTypeCustomization> MakeInstance();
+
+	FWorldBookmarkCategoryCustomization();
+
+	/** IPropertyTypeCustomization overrides */
+	virtual void CustomizeHeader(TSharedRef<IPropertyHandle> StructPropertyHandle, FDetailWidgetRow& HeaderRow, IPropertyTypeCustomizationUtils& CustomizationUtils) override;
+	virtual void CustomizeChildren(TSharedRef<IPropertyHandle> StructPropertyHandle, IDetailChildrenBuilder& ChildBuilder, IPropertyTypeCustomizationUtils& CustomizationUtils) override;
+
+private:
+	TSharedRef<SWidget> MakeCategoryComboWidget(TSharedPtr<FWorldBookmarkCategory> InItem);
+
+	void OnCategoryChanged(TSharedPtr<FWorldBookmarkCategory> NewSelection, ESelectInfo::Type SelectInfo);
+	void AddCategory(const FWorldBookmarkCategory& InCategory);
+	void SortCategories();
+	bool CreateNewCategory(FWorldBookmarkCategory& OutNewCategory);
+
+	FLinearColor GetSelectedCategoryColor() const;
+	EVisibility GetSelectedCategoryColorVisibility() const;
+	FText GetSelectedCategoryText() const;
+
+	TSharedPtr<IPropertyHandle> StructPropertyHandle;
+
+	static const FName NAME_NewCategory;
+	
+	TArray<TSharedPtr<FWorldBookmarkCategory>> KnownCategories;
+};
\ No newline at end of file
diff --git a/Engine/Source/Editor/WorldBookmark/Private/WorldBookmark/WorldBookmarkModule.cpp b/Engine/Source/Editor/WorldBookmark/Private/WorldBookmark/WorldBookmarkModule.cpp
--- a/Engine/Source/Editor/WorldBookmark/Private/WorldBookmark/WorldBookmarkModule.cpp
+++ b/Engine/Source/Editor/WorldBookmark/Private/WorldBookmark/WorldBookmarkModule.cpp
@@ -38,6 +38,9 @@ void FWorldBookmarkModule::StartupModule()
 		ClassesToUnregisterOnShutdown.Add(UWorldBookmark::StaticClass()->GetFName());
 		PropertyModule.RegisterCustomClassLayout(UWorldBookmark::StaticClass()->GetFName(), FOnGetDetailCustomizationInstance::CreateStatic(&FWorldBookmarkDetailsCustomization::MakeInstance));
 
+		StructsToUnregisterOnShutdown.Add(UWorldBookmark::StaticClass()->GetFName());
+		PropertyModule.RegisterCustomPropertyTypeLayout(FWorldBookmarkCategory::StaticStruct()->GetFName(), FOnGetPropertyTypeCustomizationInstance::CreateStatic(&FWorldBookmarkCategoryCustomization::MakeInstance));
+
 		OnAddExtraObjectsToDeleteDelegateHandle = FEditorDelegates::OnAddExtraObjectsToDelete.AddRaw(this, &FWorldBookmarkModule::OnAddExtraObjectsToDelete);
 	}
 }
@@ -62,6 +65,11 @@ void FWorldBookmarkModule::ShutdownModule()
 			{
 				PropertyModule->UnregisterCustomClassLayout(ClassName);
 			}
+
+			for (const FName& StructName : StructsToUnregisterOnShutdown)
+			{
+				PropertyModule->UnregisterCustomPropertyTypeLayout(StructName);
+			}
 		}
 
 		ClassesToUnregisterOnShutdown.Empty();
diff --git a/Engine/Source/Editor/WorldBookmark/Private/WorldBookmark/WorldBookmarkModule.h b/Engine/Source/Editor/WorldBookmark/Private/WorldBookmark/WorldBookmarkModule.h
--- a/Engine/Source/Editor/WorldBookmark/Private/WorldBookmark/WorldBookmarkModule.h
+++ b/Engine/Source/Editor/WorldBookmark/Private/WorldBookmark/WorldBookmarkModule.h
@@ -37,4 +37,5 @@ private:
 	FDelegateHandle OnAddExtraObjectsToDeleteDelegateHandle;
 
 	TArray<FName> ClassesToUnregisterOnShutdown;
+	TArray<FName> StructsToUnregisterOnShutdown;
 };
diff --git a/Engine/Source/Editor/WorldBookmark/Private/WorldBookmark/WorldBookmarkStyle.cpp b/Engine/Source/Editor/WorldBookmark/Private/WorldBookmark/WorldBookmarkStyle.cpp
--- a/Engine/Source/Editor/WorldBookmark/Private/WorldBookmark/WorldBookmarkStyle.cpp
+++ b/Engine/Source/Editor/WorldBookmark/Private/WorldBookmark/WorldBookmarkStyle.cpp
@@ -32,12 +32,16 @@ FWorldBookmarkStyle::FWorldBookmarkStyle()
 	Set("WorldBookmark.RemoveFromFavorite", new IMAGE_BRUSH("Icons/EmptyStar_16x", Icon16x16));
 	Set("WorldBookmark.PlayFromLocation", new IMAGE_BRUSH_SVG("Starship/Common/play", Icon20x20));
 	Set("WorldBookmark.MoveCameraToLocation", new IMAGE_BRUSH_SVG("Starship/EditorViewport/actor-pilot-camera", Icon16x16));
+	Set("WorldBookmark.MoveBookmarkToNewFolder", new CORE_IMAGE_BRUSH_SVG("Starship/Common/folder-plus", Icon16x16));
+	Set("WorldBookmark.CreateBookmarkInFolder", new IMAGE_BRUSH_SVG("Starship/AssetIcons/WorldBookmark_16", Icon16x16));
 
 	// Icons
-	Set("WorldBookmark.TabIcon", new IMAGE_BRUSH_SVG("Starship/AssetIcons/WorldBookmark_16", Icon16x16));
+	Set("WorldBookmark.TabIcon", new IMAGE_BRUSH_SVG("Starship/Common/Bookmarks", Icon16x16));
 	Set("WorldBookmark.IsFavorite", new IMAGE_BRUSH("Icons/Star_16x", Icon16x16));
 	Set("WorldBookmark.IsNotFavorite", new IMAGE_BRUSH("Icons/EmptyStar_16x", Icon16x16));
-	Set("WorldBookmark.RecentlyUsed", new IMAGE_BRUSH("Icons/Clock_24x", Icon16x16));
+	Set("WorldBookmark.RecentlyUsed", new CORE_IMAGE_BRUSH_SVG("Starship/Common/Recent", Icon20x20));
+	Set("WorldBookmark.FolderClosed", new CORE_IMAGE_BRUSH_SVG("Starship/Common/folder-closed", Icon16x16, FStyleColors::AccentFolder));
+	Set("WorldBookmark.FolderOpen", new CORE_IMAGE_BRUSH_SVG("Starship/Common/folder-open", Icon16x16, FStyleColors::AccentFolder));
 
 	FSlateStyleRegistry::RegisterSlateStyle(*this);
 }
diff --git a/Engine/Source/Editor/WorldBookmark/Public/WorldBookmark/WorldBookmark.h b/Engine/Source/Editor/WorldBookmark/Public/WorldBookmark/WorldBookmark.h
--- a/Engine/Source/Editor/WorldBookmark/Public/WorldBookmark/WorldBookmark.h
+++ b/Engine/Source/Editor/WorldBookmark/Public/WorldBookmark/WorldBookmark.h
@@ -7,6 +7,23 @@
 #include "EditorState/EditorStateCollection.h"
 #include "WorldBookmark.generated.h"
 
+/**
+ * World Bookmark Category
+ */
+USTRUCT()
+struct FWorldBookmarkCategory
+{
+	GENERATED_BODY()
+
+	bool operator<(const FWorldBookmarkCategory& Other) const;
+
+	UPROPERTY(EditAnywhere, Category=Category)
+	FName Name;
+
+	UPROPERTY(EditAnywhere, Category=Category, Meta=(HideAlphaChannel))
+	FColor Color;
+};
+
 /**
  * World Bookmarks are assets that stores the state of the editor world.
  */
@@ -48,12 +65,21 @@ class UWorldBookmark : public UObject
 	/** Store the last time that bookmark was loaded, in UTC. */
 	void SetUserLastLoadedTimeStampUTC(const FDateTime& InLastLoadedTimeStampUTC);
 
+	/** Retrieve the world bookmark category for this bookmark. */
+	const FWorldBookmarkCategory& GetBookmarkCategory() const;
+
     /** Get the asset registry tag used to store the matching world's name. */
 	static FName GetWorldNameAssetTag();
 
+	/** Get the asset registry tag used to store bookmark category info. */
+	static FName GetCategoryAssetTag();
+
 	/** Get the world associated with a bookmark's asset data. */
 	static FSoftObjectPath GetWorldFromAssetData(const FAssetData& InAssetData);
 
+	/** Get the bookmark category associated with a bookmark's asset data. */
+	static FWorldBookmarkCategory GetCategoryFromAssetData(const FAssetData& InAssetData);
+
 	template <typename TEditorStateType>
 	typename TEnableIf<TIsDerivedFrom<TEditorStateType, UEditorState>::IsDerived, bool>::Type HasEditorState() const
 	{
@@ -83,6 +109,9 @@ private:
 	UPROPERTY()
 	FEditorStateCollection EditorState;
 
+	UPROPERTY(EditAnywhere, Category=Category)
+	FWorldBookmarkCategory Category;
+
 	// BookmarkGuid is our key to fetch the user settings for this bookmark in the Bookmarks.ini config file.
 	// It will remain unique even if redirectors are created for the world or the bookmark itself.
 	UPROPERTY(DuplicateTransient)
diff --git a/Engine/Source/Editor/WorldBookmark/WorldBookmark.Build.cs b/Engine/Source/Editor/WorldBookmark/WorldBookmark.Build.cs
--- a/Engine/Source/Editor/WorldBookmark/WorldBookmark.Build.cs
+++ b/Engine/Source/Editor/WorldBookmark/WorldBookmark.Build.cs
@@ -20,6 +20,7 @@ public class WorldBookmark : ModuleRules
 				"Engine",
 				"InputCore",
 				"LevelEditor",
+				"Projects",
 				"PropertyEditor",
 				"SceneOutliner",
 				"SourceControl",
