World Bookmark - Moved browser classes to the UE::WorldBookmark::Browser namespace to avoid symbol clashes in monolithics builds
#buildfix

diff --git a/Engine/Source/Editor/WorldBookmark/Private/WorldBookmark/Browser/BookmarkTableRow.cpp b/Engine/Source/Editor/WorldBookmark/Private/WorldBookmark/Browser/BookmarkTableRow.cpp
--- a/Engine/Source/Editor/WorldBookmark/Private/WorldBookmark/Browser/BookmarkTableRow.cpp
+++ b/Engine/Source/Editor/WorldBookmark/Private/WorldBookmark/Browser/BookmarkTableRow.cpp
@@ -15,6 +15,9 @@
 #include "Widgets/Text/STextBlock.h"
 #include "Widgets/Views/SListView.h"
 
+namespace UE::WorldBookmark::Browser
+{
+
 #define LOCTEXT_NAMESPACE "WorldBookmarkBrowser"
 
 void FWorldBookmarkTableRow::Construct(const FArguments& InArgs, const TSharedRef<STableViewBase>& OwnerTable, FWorldBookmarkTreeItemRef InItem)
@@ -30,7 +33,7 @@ TSharedRef<SWidget> FWorldBookmarkTableRow::GenerateWidgetForColumn(const FName&
 {
 	if (FWorldBookmarkTreeItemPtr WorldBookmarkTreeItem = GetTreeItem())
 	{
-		if (ColumnName == WorldBookmarkBrowserColumn::Favorite.Id)
+		if (ColumnName == Columns::Favorite.Id)
 		{
 			if (UWorldBookmark* WorldBookmark = Cast<UWorldBookmark>(WorldBookmarkTreeItem->BookmarkAsset.GetAsset()))
 			{
@@ -40,7 +43,7 @@ TSharedRef<SWidget> FWorldBookmarkTableRow::GenerateWidgetForColumn(const FName&
 				}
 			}
 		}
-		else if (ColumnName == WorldBookmarkBrowserColumn::RecentlyUsed.Id)
+		else if (ColumnName == Columns::RecentlyUsed.Id)
 		{
 			if (UWorldBookmark* WorldBookmark = Cast<UWorldBookmark>(WorldBookmarkTreeItem->BookmarkAsset.GetAsset()))
 			{
@@ -51,7 +54,7 @@ TSharedRef<SWidget> FWorldBookmarkTableRow::GenerateWidgetForColumn(const FName&
 				}
 			}
 		}
-		else if (ColumnName == WorldBookmarkBrowserColumn::Label.Id)
+		else if (ColumnName == Columns::Label.Id)
 		{
 			if (UWorldBookmarkBrowserSettings::IsViewMode(EWorldBookmarkBrowserViewMode::TreeView))
 			{
@@ -62,7 +65,7 @@ TSharedRef<SWidget> FWorldBookmarkTableRow::GenerateWidgetForColumn(const FName&
 				return GenerateLabelForListView();
 			}
 		}
-		else if (ColumnName == WorldBookmarkBrowserColumn::World.Id)
+		else if (ColumnName == Columns::World.Id)
 		{
 			return SNew(SBox)
 				.HAlign(HAlign_Left)
@@ -72,7 +75,7 @@ TSharedRef<SWidget> FWorldBookmarkTableRow::GenerateWidgetForColumn(const FName&
 						.Text(this, &FWorldBookmarkTableRow::GetWorldNameText)
 				];
 		}
-		else if (ColumnName == WorldBookmarkBrowserColumn::Category.Id)
+		else if (ColumnName == Columns::Category.Id)
 		{
 			return SNew(SBox)
 				.HAlign(HAlign_Center)
@@ -182,3 +185,5 @@ FSlateColor FWorldBookmarkTableRow::GetIconColor() const
 }
 
 #undef LOCTEXT_NAMESPACE
+
+}
\ No newline at end of file
diff --git a/Engine/Source/Editor/WorldBookmark/Private/WorldBookmark/Browser/BookmarkTableRow.h b/Engine/Source/Editor/WorldBookmark/Private/WorldBookmark/Browser/BookmarkTableRow.h
--- a/Engine/Source/Editor/WorldBookmark/Private/WorldBookmark/Browser/BookmarkTableRow.h
+++ b/Engine/Source/Editor/WorldBookmark/Private/WorldBookmark/Browser/BookmarkTableRow.h
@@ -7,6 +7,9 @@
 #include "WorldBookmark/Browser/BookmarkTreeItem.h"
 #include "WorldBookmark/Browser/CommonTableRow.h"
 
+namespace UE::WorldBookmark::Browser
+{
+
 class FWorldBookmarkTableRow : public SMultiColumnTableRow<FWorldBookmarkTreeItemPtr>, public FWorldBookmarkTableRowBase<FWorldBookmarkTreeItem>
 {
 public:
@@ -25,3 +28,5 @@ public:
 
 	FSlateColor GetIconColor() const;
 };
+
+}
\ No newline at end of file
diff --git a/Engine/Source/Editor/WorldBookmark/Private/WorldBookmark/Browser/BookmarkTreeItem.cpp b/Engine/Source/Editor/WorldBookmark/Private/WorldBookmark/Browser/BookmarkTreeItem.cpp
--- a/Engine/Source/Editor/WorldBookmark/Private/WorldBookmark/Browser/BookmarkTreeItem.cpp
+++ b/Engine/Source/Editor/WorldBookmark/Private/WorldBookmark/Browser/BookmarkTreeItem.cpp
@@ -9,6 +9,9 @@
 #include "IContentBrowserSingleton.h"
 #include "ObjectTools.h"
 
+namespace UE::WorldBookmark::Browser
+{
+
 const ETreeItemType FWorldBookmarkTreeItem::Type(ETreeItemType::WorldBookmark);
 
 FName FWorldBookmarkTreeItem::GetName() const
@@ -80,4 +83,6 @@ void FWorldBookmarkTreeItem::ShowInContentBrowser() const
 {
 	FContentBrowserModule& ContentBrowserModule = FModuleManager::Get().LoadModuleChecked<FContentBrowserModule>("ContentBrowser");
 	ContentBrowserModule.Get().SyncBrowserToAssets(TArray<FAssetData>({ BookmarkAsset }));
-}
\ No newline at end of file
+}
+
+}
diff --git a/Engine/Source/Editor/WorldBookmark/Private/WorldBookmark/Browser/BookmarkTreeItem.h b/Engine/Source/Editor/WorldBookmark/Private/WorldBookmark/Browser/BookmarkTreeItem.h
--- a/Engine/Source/Editor/WorldBookmark/Private/WorldBookmark/Browser/BookmarkTreeItem.h
+++ b/Engine/Source/Editor/WorldBookmark/Private/WorldBookmark/Browser/BookmarkTreeItem.h
@@ -5,6 +5,9 @@
 
 #include "AssetRegistry/AssetData.h"
 
+namespace UE::WorldBookmark::Browser
+{
+
 struct FWorldBookmarkTreeItem : public ITreeItem
 {
 public:
@@ -38,3 +41,5 @@ private:
 
 typedef TSharedPtr<FWorldBookmarkTreeItem> FWorldBookmarkTreeItemPtr;
 typedef TSharedRef<FWorldBookmarkTreeItem> FWorldBookmarkTreeItemRef;
+
+}
diff --git a/Engine/Source/Editor/WorldBookmark/Private/WorldBookmark/Browser/Columns.cpp b/Engine/Source/Editor/WorldBookmark/Private/WorldBookmark/Browser/Columns.cpp
--- a/Engine/Source/Editor/WorldBookmark/Private/WorldBookmark/Browser/Columns.cpp
+++ b/Engine/Source/Editor/WorldBookmark/Private/WorldBookmark/Browser/Columns.cpp
@@ -4,9 +4,12 @@
 
 #include "Internationalization/Internationalization.h"
 
+namespace UE::WorldBookmark::Browser
+{
+
 #define LOCTEXT_NAMESPACE "WorldBookmarkBrowser"
 
-namespace WorldBookmarkBrowserColumn
+namespace Columns
 {
 	const FColumnDefinition Favorite =
 	{
@@ -47,3 +50,5 @@ namespace WorldBookmarkBrowserColumn
 }
 
 #undef LOCTEXT_NAMESPACE
+
+}
diff --git a/Engine/Source/Editor/WorldBookmark/Private/WorldBookmark/Browser/Columns.h b/Engine/Source/Editor/WorldBookmark/Private/WorldBookmark/Browser/Columns.h
--- a/Engine/Source/Editor/WorldBookmark/Private/WorldBookmark/Browser/Columns.h
+++ b/Engine/Source/Editor/WorldBookmark/Private/WorldBookmark/Browser/Columns.h
@@ -5,7 +5,7 @@
 #include "Internationalization/Text.h"
 #include "UObject/NameTypes.h"
 
-namespace WorldBookmarkBrowserColumn
+namespace UE::WorldBookmark::Browser::Columns
 {
 	struct FColumnDefinition
 	{
diff --git a/Engine/Source/Editor/WorldBookmark/Private/WorldBookmark/Browser/CommonTableRow.h b/Engine/Source/Editor/WorldBookmark/Private/WorldBookmark/Browser/CommonTableRow.h
--- a/Engine/Source/Editor/WorldBookmark/Private/WorldBookmark/Browser/CommonTableRow.h
+++ b/Engine/Source/Editor/WorldBookmark/Private/WorldBookmark/Browser/CommonTableRow.h
@@ -13,6 +13,9 @@
 #include "WorldBookmark/Browser/TreeItem.h"
 #include "WorldBookmark/WorldBookmarkStyle.h"
 
+namespace UE::WorldBookmark::Browser
+{
+
 template <class TItemType = ITreeItem>
 class FWorldBookmarkTableRowBase
 {
@@ -231,3 +234,5 @@ TSharedRef<SWidget> CreateTreeLabelWidget(TSharedRef<TTableRow> TableRow)
 		]
 	];
 }
+
+}
diff --git a/Engine/Source/Editor/WorldBookmark/Private/WorldBookmark/Browser/DragDropOperation.cpp b/Engine/Source/Editor/WorldBookmark/Private/WorldBookmark/Browser/DragDropOperation.cpp
--- a/Engine/Source/Editor/WorldBookmark/Private/WorldBookmark/Browser/DragDropOperation.cpp
+++ b/Engine/Source/Editor/WorldBookmark/Private/WorldBookmark/Browser/DragDropOperation.cpp
@@ -7,6 +7,9 @@
 #include "WorldBookmark/Browser/Settings.h"
 #include "WorldBookmark/WorldBookmarkStyle.h"
 
+namespace UE::WorldBookmark::Browser
+{
+
 TSharedPtr<FTableRowDragDropOp> FTableRowDragDropOp::New(FTreeItemPtr InTreeItem)
 {
 	if (UWorldBookmarkBrowserSettings::IsViewMode(EWorldBookmarkBrowserViewMode::TreeView))
@@ -61,3 +64,4 @@ FReply FTableRowDragDropOp::AcceptDrop(FTreeItemPtr DropTarget) const
 	return FReply::Handled().EndDragDrop();
 }
 
+}
diff --git a/Engine/Source/Editor/WorldBookmark/Private/WorldBookmark/Browser/DragDropOperation.h b/Engine/Source/Editor/WorldBookmark/Private/WorldBookmark/Browser/DragDropOperation.h
--- a/Engine/Source/Editor/WorldBookmark/Private/WorldBookmark/Browser/DragDropOperation.h
+++ b/Engine/Source/Editor/WorldBookmark/Private/WorldBookmark/Browser/DragDropOperation.h
@@ -6,6 +6,9 @@
 #include "Widgets/Views/STableRow.h"
 #include "WorldBookmark/Browser/TreeItem.h"
 
+namespace UE::WorldBookmark::Browser
+{
+
 class FTableRowDragDropOp : public FDecoratedDragDropOp
 {
 public:
@@ -17,3 +20,5 @@ public:
 
 	TWeakPtr<ITreeItem> DraggedItem;
 };
+
+}
diff --git a/Engine/Source/Editor/WorldBookmark/Private/WorldBookmark/Browser/FolderTableRow.cpp b/Engine/Source/Editor/WorldBookmark/Private/WorldBookmark/Browser/FolderTableRow.cpp
--- a/Engine/Source/Editor/WorldBookmark/Private/WorldBookmark/Browser/FolderTableRow.cpp
+++ b/Engine/Source/Editor/WorldBookmark/Private/WorldBookmark/Browser/FolderTableRow.cpp
@@ -7,6 +7,9 @@
 #include "Widgets/Text/SInlineEditableTextBlock.h"
 #include "Widgets/Text/STextBlock.h"
 
+namespace UE::WorldBookmark::Browser
+{
+
 void FFolderTableRow::Construct(const FArguments& InArgs, const TSharedRef<STableViewBase>& InOwnerTableView, FFolderTreeItemRef InItem)
 {
 	WeakTreeItem = InItem;
@@ -22,3 +25,5 @@ void FFolderTableRow::Construct(const FArguments& InArgs, const TSharedRef<STabl
 		CreateTreeLabelWidget(SharedThis(this))
 	];
 }
+
+}
diff --git a/Engine/Source/Editor/WorldBookmark/Private/WorldBookmark/Browser/FolderTableRow.h b/Engine/Source/Editor/WorldBookmark/Private/WorldBookmark/Browser/FolderTableRow.h
--- a/Engine/Source/Editor/WorldBookmark/Private/WorldBookmark/Browser/FolderTableRow.h
+++ b/Engine/Source/Editor/WorldBookmark/Private/WorldBookmark/Browser/FolderTableRow.h
@@ -7,8 +7,13 @@
 #include "WorldBookmark/Browser/FolderTreeItem.h"
 #include "WorldBookmark/Browser/CommonTableRow.h"
 
+namespace UE::WorldBookmark::Browser
+{
+
 class FFolderTableRow : public STableRow<FFolderTreeItemPtr>, public FWorldBookmarkTableRowBase<FFolderTreeItem>
 {
 public:
 	void Construct(const FArguments& InArgs, const TSharedRef<STableViewBase>& InOwnerTableView, FFolderTreeItemRef InItem);
 };
+
+}
diff --git a/Engine/Source/Editor/WorldBookmark/Private/WorldBookmark/Browser/FolderTreeItem.cpp b/Engine/Source/Editor/WorldBookmark/Private/WorldBookmark/Browser/FolderTreeItem.cpp
--- a/Engine/Source/Editor/WorldBookmark/Private/WorldBookmark/Browser/FolderTreeItem.cpp
+++ b/Engine/Source/Editor/WorldBookmark/Private/WorldBookmark/Browser/FolderTreeItem.cpp
@@ -10,6 +10,9 @@
 #include "IContentBrowserSingleton.h"
 #include "ObjectTools.h"
 
+namespace UE::WorldBookmark::Browser
+{
+
 const ETreeItemType FFolderTreeItem::Type(ETreeItemType::Folder);
 
 
@@ -294,4 +297,6 @@ void FFolderTreeItem::Sort(TFunctionRef<bool(const FTreeItemPtr& A, const FTreeI
 void FFolderTreeItem::Reset()
 {
 	Children.Reset();
-}
\ No newline at end of file
+}
+
+}
diff --git a/Engine/Source/Editor/WorldBookmark/Private/WorldBookmark/Browser/FolderTreeItem.h b/Engine/Source/Editor/WorldBookmark/Private/WorldBookmark/Browser/FolderTreeItem.h
--- a/Engine/Source/Editor/WorldBookmark/Private/WorldBookmark/Browser/FolderTreeItem.h
+++ b/Engine/Source/Editor/WorldBookmark/Private/WorldBookmark/Browser/FolderTreeItem.h
@@ -3,6 +3,9 @@
 
 #include "WorldBookmark/Browser/TreeItem.h"
 
+namespace UE::WorldBookmark::Browser
+{
+
 struct FFolderTreeItem : public ITreeItem
 {
 public:
@@ -88,3 +91,5 @@ private:
 
 typedef TSharedPtr<FFolderTreeItem> FFolderTreeItemPtr;
 typedef TSharedRef<FFolderTreeItem> FFolderTreeItemRef;
+
+}
diff --git a/Engine/Source/Editor/WorldBookmark/Private/WorldBookmark/Browser/Icons.cpp b/Engine/Source/Editor/WorldBookmark/Private/WorldBookmark/Browser/Icons.cpp
--- a/Engine/Source/Editor/WorldBookmark/Private/WorldBookmark/Browser/Icons.cpp
+++ b/Engine/Source/Editor/WorldBookmark/Private/WorldBookmark/Browser/Icons.cpp
@@ -7,6 +7,9 @@
 #include "Widgets/SBoxPanel.h"
 #include "WorldBookmark/WorldBookmarkStyle.h"
 
+namespace UE::WorldBookmark::Browser
+{
+
 FSlateIcon GetFavoriteWorldBookmarkIcon(bool bIsFavorite)
 {
 	const FName FavoriteIconName = bIsFavorite ? "WorldBookmark.IsFavorite" : "WorldBookmark.IsNotFavorite";
@@ -43,3 +46,5 @@ TSharedRef<SWidget> CreateIconWidget(const FSlateIcon& SlateIcon)
 				.Image(SlateIcon.GetIcon())
 		];
 }
+
+}
diff --git a/Engine/Source/Editor/WorldBookmark/Private/WorldBookmark/Browser/Icons.h b/Engine/Source/Editor/WorldBookmark/Private/WorldBookmark/Browser/Icons.h
--- a/Engine/Source/Editor/WorldBookmark/Private/WorldBookmark/Browser/Icons.h
+++ b/Engine/Source/Editor/WorldBookmark/Private/WorldBookmark/Browser/Icons.h
@@ -6,8 +6,11 @@
 
 class SWidget;
 
-FSlateIcon GetFavoriteWorldBookmarkIcon(bool bIsFavorite);
-FSlateIcon GetRecentlyUsedWorldBookmarkIcon();
-FSlateIcon GetFolderIcon(bool bIsExpanded);
-FSlateIcon GetWorldBookmarkIcon();
-TSharedRef<SWidget> CreateIconWidget(const FSlateIcon& SlateIcon);
+namespace UE::WorldBookmark::Browser
+{
+	FSlateIcon GetFavoriteWorldBookmarkIcon(bool bIsFavorite);
+	FSlateIcon GetRecentlyUsedWorldBookmarkIcon();
+	FSlateIcon GetFolderIcon(bool bIsExpanded);
+	FSlateIcon GetWorldBookmarkIcon();
+	TSharedRef<SWidget> CreateIconWidget(const FSlateIcon& SlateIcon);
+}
\ No newline at end of file
diff --git a/Engine/Source/Editor/WorldBookmark/Private/WorldBookmark/Browser/SWorldBookmarkBrowser.cpp b/Engine/Source/Editor/WorldBookmark/Private/WorldBookmark/Browser/SWorldBookmarkBrowser.cpp
--- a/Engine/Source/Editor/WorldBookmark/Private/WorldBookmark/Browser/SWorldBookmarkBrowser.cpp
+++ b/Engine/Source/Editor/WorldBookmark/Private/WorldBookmark/Browser/SWorldBookmarkBrowser.cpp
@@ -45,6 +45,9 @@
 #include "Widgets/Text/SInlineEditableTextBlock.h"
 #include "Widgets/Views/SHeaderRow.h"
 
+namespace UE::WorldBookmark::Browser
+{
+
 #define LOCTEXT_NAMESPACE "WorldBookmarkBrowser"
 
 static UWorld* GetCurrentWorld()
@@ -79,8 +82,8 @@ SWorldBookmarkBrowser::SWorldBookmarkBrowser()
 	FLevelEditorModule& LevelEditorModule = FModuleManager::LoadModuleChecked<FLevelEditorModule>(TEXT("LevelEditor"));
 	OnMapChangedHandle = LevelEditorModule.OnMapChanged().AddRaw(this, &SWorldBookmarkBrowser::OnMapChanged);
 
-	ColumnsSortingParams[EColumnSortPriority::Primary] = { WorldBookmarkBrowserColumn::Favorite.Id, EColumnSortMode::Ascending };
-	ColumnsSortingParams[EColumnSortPriority::Secondary] = { WorldBookmarkBrowserColumn::Label.Id, EColumnSortMode::Ascending };
+	ColumnsSortingParams[EColumnSortPriority::Primary] = { Columns::Favorite.Id, EColumnSortMode::Ascending };
+	ColumnsSortingParams[EColumnSortPriority::Secondary] = { Columns::Label.Id, EColumnSortMode::Ascending };
 
 	OnSettingsChangedHandle = UWorldBookmarkBrowserSettings::OnSettingsChanged().AddRaw(this, &SWorldBookmarkBrowser::OnSettingsChanged);
 }
@@ -277,7 +280,7 @@ void SWorldBookmarkBrowser::RecreateColumns()
 {
 	HeaderRow->ClearColumns();
 
-	auto HeaderRowColumn = [this](const WorldBookmarkBrowserColumn::FColumnDefinition& InColumn) -> SHeaderRow::FColumn::FArguments
+	auto HeaderRowColumn = [this](const Columns::FColumnDefinition& InColumn) -> SHeaderRow::FColumn::FArguments
 	{
 		return SHeaderRow::Column(InColumn.Id)
 			.SortMode(this, &SWorldBookmarkBrowser::GetColumnSortMode, InColumn.Id)
@@ -287,7 +290,7 @@ void SWorldBookmarkBrowser::RecreateColumns()
 			.DefaultTooltip(InColumn.ToolTipText);
 	};
 
-	auto HeaderRowIconColumn = [this, &HeaderRowColumn](const WorldBookmarkBrowserColumn::FColumnDefinition& InColumn)->SHeaderRow::FColumn::FArguments
+	auto HeaderRowIconColumn = [this, &HeaderRowColumn](const Columns::FColumnDefinition& InColumn)->SHeaderRow::FColumn::FArguments
 	{
 		return HeaderRowColumn(InColumn)
 			.FixedWidth(24.f)
@@ -301,31 +304,31 @@ void SWorldBookmarkBrowser::RecreateColumns()
 
 	// Favorite
 	SHeaderRow::FColumn::FArguments FavoriteColumn =
-		HeaderRowIconColumn(WorldBookmarkBrowserColumn::Favorite)
+		HeaderRowIconColumn(Columns::Favorite)
 		[
 			CreateIconWidget(GetFavoriteWorldBookmarkIcon(true))
 		];
 
 	// RecentlyUsed
 	SHeaderRow::FColumn::FArguments RecentlyUsedColumn =
-		HeaderRowIconColumn(WorldBookmarkBrowserColumn::RecentlyUsed)
+		HeaderRowIconColumn(Columns::RecentlyUsed)
 		[
 			CreateIconWidget(GetRecentlyUsedWorldBookmarkIcon())
 		];
 
 	// LabelColumn
 	SHeaderRow::FColumn::FArguments LabelColumn =
-		HeaderRowColumn(WorldBookmarkBrowserColumn::Label)
+		HeaderRowColumn(Columns::Label)
 		.ShouldGenerateWidget(true); // Ensure the column cannot be hidden (grayed out in the show/hide drop down menu)
 
 	// World
 	SHeaderRow::FColumn::FArguments WorldColumn =
-		HeaderRowColumn(WorldBookmarkBrowserColumn::World)
+		HeaderRowColumn(Columns::World)
 		.FillWidth(0.65f);
 
 	// Category
 	SHeaderRow::FColumn::FArguments CategoryColumn =
-		HeaderRowIconColumn(WorldBookmarkBrowserColumn::Category)
+		HeaderRowIconColumn(Columns::Category)
 		[
 			SNew(SColorBlock)
 				.Color(FStyleColors::Foreground.GetSpecifiedColor())
@@ -686,8 +689,8 @@ void SWorldBookmarkBrowser::SortItems()
 {
 	auto CompareFolderTreeItems = [this](const FFolderTreeItem* A, const FFolderTreeItem* B) -> bool
 	{
-		const bool bReverseSort = (ColumnsSortingParams[EColumnSortPriority::Primary].Key == WorldBookmarkBrowserColumn::Label.Id && ColumnsSortingParams[EColumnSortPriority::Primary].Value == EColumnSortMode::Descending) ||
-							(ColumnsSortingParams[EColumnSortPriority::Secondary].Key == WorldBookmarkBrowserColumn::Label.Id && ColumnsSortingParams[EColumnSortPriority::Secondary].Value == EColumnSortMode::Descending);
+		const bool bReverseSort = (ColumnsSortingParams[EColumnSortPriority::Primary].Key == Columns::Label.Id && ColumnsSortingParams[EColumnSortPriority::Primary].Value == EColumnSortMode::Descending) ||
+							(ColumnsSortingParams[EColumnSortPriority::Secondary].Key == Columns::Label.Id && ColumnsSortingParams[EColumnSortPriority::Secondary].Value == EColumnSortMode::Descending);
 
 		return bReverseSort ? B->GetName().LexicalLess(A->GetName()) : A->GetName().LexicalLess(B->GetName());
 	};
@@ -709,23 +712,23 @@ void SWorldBookmarkBrowser::SortItems()
 				return WorldBookmarkA ? -1 : 0;
 			}
 
-			if (SortOrder.Key == WorldBookmarkBrowserColumn::Favorite.Id)
+			if (SortOrder.Key == Columns::Favorite.Id)
 			{
 				return ApplySortMode(WorldBookmarkA->GetIsUserFavorite() == WorldBookmarkB->GetIsUserFavorite() ? 0 : WorldBookmarkA->GetIsUserFavorite() ? -1 : 1);
 			}
-			else if (SortOrder.Key == WorldBookmarkBrowserColumn::RecentlyUsed.Id)
+			else if (SortOrder.Key == Columns::RecentlyUsed.Id)
 			{
 				return ApplySortMode(WorldBookmarkA->GetUserLastLoadedTimeStampUTC() == WorldBookmarkB->GetUserLastLoadedTimeStampUTC() ? 0 : WorldBookmarkA->GetUserLastLoadedTimeStampUTC() < WorldBookmarkB->GetUserLastLoadedTimeStampUTC() ? -1 : 1);
 			}
-			else if (SortOrder.Key == WorldBookmarkBrowserColumn::Label.Id)
+			else if (SortOrder.Key == Columns::Label.Id)
 			{
 				return ApplySortMode(WorldBookmarkA->GetName() == WorldBookmarkB->GetName() ? 0 : WorldBookmarkA->GetName() < WorldBookmarkB->GetName() ? -1 : 1);
 			}
-			else if (SortOrder.Key == WorldBookmarkBrowserColumn::World.Id)
+			else if (SortOrder.Key == Columns::World.Id)
 			{
 				return ApplySortMode(A->GetBookmarkWorld().GetAssetName() == B->GetBookmarkWorld().GetAssetName() ? 0 : A->GetBookmarkWorld().GetAssetName() < B->GetBookmarkWorld().GetAssetName() ? -1 : 1);
 			}
-			else if (SortOrder.Key == WorldBookmarkBrowserColumn::Category.Id)
+			else if (SortOrder.Key == Columns::Category.Id)
 			{
 				return ApplySortMode(WorldBookmarkA->GetBookmarkCategory().Name == WorldBookmarkB->GetBookmarkCategory().Name ? 0 : WorldBookmarkA->GetBookmarkCategory() < WorldBookmarkB->GetBookmarkCategory() ? -1 : 1);
 			}
@@ -1489,3 +1492,4 @@ void SWorldBookmarkBrowser::CreateBookmarkInSelectedFolder()
 
 #undef LOCTEXT_NAMESPACE
 
+}
diff --git a/Engine/Source/Editor/WorldBookmark/Private/WorldBookmark/Browser/SWorldBookmarkBrowser.h b/Engine/Source/Editor/WorldBookmark/Private/WorldBookmark/Browser/SWorldBookmarkBrowser.h
--- a/Engine/Source/Editor/WorldBookmark/Private/WorldBookmark/Browser/SWorldBookmarkBrowser.h
+++ b/Engine/Source/Editor/WorldBookmark/Private/WorldBookmark/Browser/SWorldBookmarkBrowser.h
@@ -15,6 +15,9 @@ class UWorldBookmark;
 class IDetailsView;
 class SAssetView;
 
+namespace UE::WorldBookmark::Browser
+{
+
 class SWorldBookmarkBrowser : public SCompoundWidget
 {
 public:
@@ -133,4 +136,6 @@ private:
 	bool bIsInTick;
 	bool bPendingRefresh;
 	bool bExpandAllItemsOnNextRefresh;
-};
\ No newline at end of file
+};
+
+}
diff --git a/Engine/Source/Editor/WorldBookmark/Private/WorldBookmark/Browser/Settings.cpp b/Engine/Source/Editor/WorldBookmark/Private/WorldBookmark/Browser/Settings.cpp
--- a/Engine/Source/Editor/WorldBookmark/Private/WorldBookmark/Browser/Settings.cpp
+++ b/Engine/Source/Editor/WorldBookmark/Private/WorldBookmark/Browser/Settings.cpp
@@ -12,7 +12,7 @@ UWorldBookmarkBrowserSettings& UWorldBookmarkBrowserSettings::Get()
 	{
 		Instance = NewObject<UWorldBookmarkBrowserSettings>();
 
-		Instance->HiddenColumns = WorldBookmarkBrowserColumn::DefaultHiddenColumns;
+		Instance->HiddenColumns = UE::WorldBookmark::Browser::Columns::DefaultHiddenColumns;
 
 		Instance->LoadEditorConfig();
 		Instance->AddToRoot();
@@ -129,4 +129,4 @@ void UWorldBookmarkBrowserSettings::SetViewMode(EWorldBookmarkBrowserViewMode In
 		Settings.ViewMode = InViewMode;
 		Settings.ApplySettingsChanges();
 	}
-}
\ No newline at end of file
+}
diff --git a/Engine/Source/Editor/WorldBookmark/Private/WorldBookmark/Browser/TreeItem.cpp b/Engine/Source/Editor/WorldBookmark/Private/WorldBookmark/Browser/TreeItem.cpp
--- a/Engine/Source/Editor/WorldBookmark/Private/WorldBookmark/Browser/TreeItem.cpp
+++ b/Engine/Source/Editor/WorldBookmark/Private/WorldBookmark/Browser/TreeItem.cpp
@@ -2,6 +2,9 @@
 
 #include "WorldBookmark/Browser/TreeItem.h"
 
+namespace UE::WorldBookmark::Browser
+{
+
 ITreeItem::ITreeItem(ETreeItemType InType)
 	: TreeItemType(InType)
 	, ParentItem(nullptr)
@@ -37,4 +40,6 @@ FString ITreeItem::GetAssetPath() const
 	} while (CurrentItem = CurrentItem->GetParent(), CurrentItem);
 
 	return FString(FullPathBuilder);
-}
\ No newline at end of file
+}
+
+}
diff --git a/Engine/Source/Editor/WorldBookmark/Private/WorldBookmark/Browser/TreeItem.h b/Engine/Source/Editor/WorldBookmark/Private/WorldBookmark/Browser/TreeItem.h
--- a/Engine/Source/Editor/WorldBookmark/Private/WorldBookmark/Browser/TreeItem.h
+++ b/Engine/Source/Editor/WorldBookmark/Private/WorldBookmark/Browser/TreeItem.h
@@ -4,6 +4,9 @@
 #include "Delegates/Delegate.h"
 #include "Templates/SharedPointer.h"
 
+namespace UE::WorldBookmark::Browser
+{
+
 enum ETreeItemType
 {
 	Folder,
@@ -86,3 +89,5 @@ protected:
 
 typedef TSharedPtr<ITreeItem> FTreeItemPtr;
 typedef TSharedRef<ITreeItem> FTreeItemRef;
+
+}
diff --git a/Engine/Source/Editor/WorldBookmark/Private/WorldBookmark/WorldBookmarkModule.cpp b/Engine/Source/Editor/WorldBookmark/Private/WorldBookmark/WorldBookmarkModule.cpp
--- a/Engine/Source/Editor/WorldBookmark/Private/WorldBookmark/WorldBookmarkModule.cpp
+++ b/Engine/Source/Editor/WorldBookmark/Private/WorldBookmark/WorldBookmarkModule.cpp
@@ -106,7 +106,7 @@ TSharedRef<SDockTab> FWorldBookmarkModule::SpawnWorldBookmarkBrowserTab(const FS
 
 TSharedRef<SWidget> FWorldBookmarkModule::CreateWorldBookmarkBrowser()
 {
-	return SNew(SWorldBookmarkBrowser);
+	return SNew(UE::WorldBookmark::Browser::SWorldBookmarkBrowser);
 }
 
 void FWorldBookmarkModule::OnAddExtraObjectsToDelete(const TArray<UObject*>& InObjectsToDelete, TSet<UObject*>& OutSecondaryObjects)
