World Bookmark - Added level bookmark feature
* Can now specify a World Bookmark as the level bookmark through the world settings
* The level bookmark will be applied on map load (only in the editor and not while running commandlets)
* Added user settings to control this feature (on by default)
* Added validation to ensure the specified level bookmark is referencing the correct world

diff --git a/Engine/Source/Editor/UnrealEd/Private/EditorState/EditorStateSubsystem.cpp b/Engine/Source/Editor/UnrealEd/Private/EditorState/EditorStateSubsystem.cpp
--- a/Engine/Source/Editor/UnrealEd/Private/EditorState/EditorStateSubsystem.cpp
+++ b/Engine/Source/Editor/UnrealEd/Private/EditorState/EditorStateSubsystem.cpp
@@ -147,6 +147,10 @@ void UEditorStateSubsystem::RestoreEditorState(const FEditorStateCollection& InS
 
 void UEditorStateSubsystem::RestoreEditorState(const FEditorStateCollection& InState, const TArray<TSubclassOf<UEditorState>>& InEditorStateTypeFilter)
 {
+	check(!bIsRestoringEditorState);
+
+	TGuardValue<bool> RestoringEditorState(bIsRestoringEditorState, true);
+
 	UE_LOG(LogEditorState, Log, TEXT("Restoring editor state..."));
 
 	InState.ForEachState([this](const UEditorState* StateToRestore, bool bRestoredDependencies)
@@ -167,4 +171,9 @@ void UEditorStateSubsystem::RestoreEditorState(const FEditorStateCollection& InS
 	UE_LOG(LogEditorState, Log, TEXT("Restored editor state... DONE"));
 }
 
+bool UEditorStateSubsystem::IsRestoringEditorState() const
+{
+	return bIsRestoringEditorState;
+}
+
 #undef LOCTEXT_NAMESPACE
\ No newline at end of file
diff --git a/Engine/Source/Editor/UnrealEd/Public/EditorState/EditorStateSubsystem.h b/Engine/Source/Editor/UnrealEd/Public/EditorState/EditorStateSubsystem.h
--- a/Engine/Source/Editor/UnrealEd/Public/EditorState/EditorStateSubsystem.h
+++ b/Engine/Source/Editor/UnrealEd/Public/EditorState/EditorStateSubsystem.h
@@ -62,7 +62,12 @@ public:
 	 * @param InState		Collection of states to restore.
 	 */
 	UNREALED_API void RestoreEditorState(const FEditorStateCollection& InState, const TArray<TSubclassOf<UEditorState>>& InEditorStateTypeFilter);
-	
+
+	/**
+	 * Returns true if the subsystem is currently restoring an editor state.
+	 */
+	UNREALED_API bool IsRestoringEditorState() const;
+		
 private:
 	UNREALED_API void RegisterEditorStateType(TSubclassOf<UEditorState> InEditorStateType);
 	UNREALED_API void UnregisterEditorStateType(TSubclassOf<UEditorState> InEditorStateType);
@@ -70,4 +75,6 @@ private:
 	void OutputOperationResult(const UEditorState* InState, bool bRestore, const UEditorState::FOperationResult& InOperationResult) const;
 
 	TArray<TSubclassOf<UEditorState>> RegisteredEditorStateTypes;
+
+	bool bIsRestoringEditorState = false;
 };
diff --git a/Engine/Source/Editor/WorldBookmark/Private/WorldBookmark/WorldBookmarkEditorSettings.cpp b/Engine/Source/Editor/WorldBookmark/Private/WorldBookmark/WorldBookmarkEditorSettings.cpp
new file mode 100644
--- /dev/null
+++ b/Engine/Source/Editor/WorldBookmark/Private/WorldBookmark/WorldBookmarkEditorSettings.cpp
@@ -0,0 +1,8 @@
+// Copyright Epic Games, Inc. All Rights Reserved.
+
+#include "WorldBookmark/WorldBookmarkEditorSettings.h"
+
+UWorldBookmarkEditorPerProjectUserSettings::UWorldBookmarkEditorPerProjectUserSettings()
+{
+	bEnableDefaultBookmarks = true;
+}
diff --git a/Engine/Source/Editor/WorldBookmark/Private/WorldBookmark/WorldBookmarkEditorSettings.h b/Engine/Source/Editor/WorldBookmark/Private/WorldBookmark/WorldBookmarkEditorSettings.h
new file mode 100644
--- /dev/null
+++ b/Engine/Source/Editor/WorldBookmark/Private/WorldBookmark/WorldBookmarkEditorSettings.h
@@ -0,0 +1,20 @@
+// Copyright Epic Games, Inc. All Rights Reserved.
+
+#pragma once
+
+#include "CoreMinimal.h"
+#include "Engine/DeveloperSettings.h"
+#include "WorldBookmarkEditorSettings.generated.h"
+
+UCLASS(config = EditorPerProjectUserSettings, meta = (DisplayName = "World Bookmark"))
+class UWorldBookmarkEditorPerProjectUserSettings : public UDeveloperSettings
+{
+	GENERATED_BODY()
+
+public:
+	UWorldBookmarkEditorPerProjectUserSettings();
+
+	/** When enabled, the default bookmark will be applied when loading a level */
+	UPROPERTY(config, EditAnywhere, Category="Default Bookmark")
+	bool bEnableDefaultBookmarks;
+};
diff --git a/Engine/Source/Editor/WorldBookmark/Private/WorldBookmark/WorldBookmarkModule.cpp b/Engine/Source/Editor/WorldBookmark/Private/WorldBookmark/WorldBookmarkModule.cpp
--- a/Engine/Source/Editor/WorldBookmark/Private/WorldBookmark/WorldBookmarkModule.cpp
+++ b/Engine/Source/Editor/WorldBookmark/Private/WorldBookmark/WorldBookmarkModule.cpp
@@ -4,17 +4,23 @@
 #include "WorldBookmark/WorldBookmark.h"
 #include "WorldBookmark/WorldBookmarkCommands.h"
 #include "WorldBookmark/WorldBookmarkDetailsCustomization.h"
+#include "WorldBookmark/WorldBookmarkEditorSettings.h"
 #include "WorldBookmark/WorldBookmarkStyle.h"
 #include "WorldBookmark/Browser/SWorldBookmarkBrowser.h"
 
 #include "AssetRegistry/IAssetRegistry.h"
 #include "AssetRegistry/AssetData.h"
+#include "EditorState/EditorStateSubsystem.h"
+#include "EditorState/WorldEditorState.h"
+#include "Framework/Notifications/NotificationManager.h"
+#include "GameFramework/WorldSettings.h"
 #include "LevelEditor.h"
 #include "Misc/MessageDialog.h"
 #include "Modules/ModuleManager.h"
 #include "ObjectTools.h"
 #include "PropertyEditorModule.h"
 #include "Widgets/Docking/SDockTab.h"
+#include "Widgets/Notifications/SNotificationList.h"
 #include "WorkspaceMenuStructure.h"
 #include "WorkspaceMenuStructureModule.h"
 
@@ -42,6 +48,12 @@ void FWorldBookmarkModule::StartupModule()
 		PropertyModule.RegisterCustomPropertyTypeLayout(FWorldBookmarkCategory::StaticStruct()->GetFName(), FOnGetPropertyTypeCustomizationInstance::CreateStatic(&FWorldBookmarkCategoryCustomization::MakeInstance));
 
 		OnAddExtraObjectsToDeleteDelegateHandle = FEditorDelegates::OnAddExtraObjectsToDelete.AddRaw(this, &FWorldBookmarkModule::OnAddExtraObjectsToDelete);
+
+		// Listen for map change events
+		OnMapChangedHandle = LevelEditorModule.OnMapChanged().AddRaw(this, &FWorldBookmarkModule::OnMapChanged);
+
+		// Validate level default bookmark changes
+		OnDefaultBookmarkChangedHandle = AWorldSettings::OnDefaultBookmarkChanged.AddRaw(this, &FWorldBookmarkModule::OnDefaultBookmarkChanged);
 	}
 }
 
@@ -57,6 +69,8 @@ void FWorldBookmarkModule::ShutdownModule()
 			{
 				TabManager->UnregisterTabSpawner(WorldBookmarkBrowserTabId);
 			}
+
+			LevelEditorModule->OnMapChanged().Remove(OnMapChangedHandle);
 		}		
 
 		if (FPropertyEditorModule* PropertyModule = FModuleManager::GetModulePtr<FPropertyEditorModule>("PropertyEditor"))
@@ -72,7 +86,10 @@ void FWorldBookmarkModule::ShutdownModule()
 			}
 		}
 
+		AWorldSettings::OnDefaultBookmarkChanged.Remove(OnDefaultBookmarkChangedHandle);
+
 		ClassesToUnregisterOnShutdown.Empty();
+		StructsToUnregisterOnShutdown.Empty();
 
 		FEditorDelegates::OnAddExtraObjectsToDelete.Remove(OnAddExtraObjectsToDeleteDelegateHandle);
 	}
@@ -173,4 +190,75 @@ void FWorldBookmarkModule::OnAddExtraObjectsToDelete(const TArray<UObject*>& InO
 	}
 }
 
+void FWorldBookmarkModule::OnMapChanged(UWorld* InWorld, EMapChangeType InMapChangeType)
+{
+	if (InMapChangeType != EMapChangeType::LoadMap)
+	{
+		return;
+	}
+
+	// Load the default bookmark if the user has that option enabled
+	const UWorldBookmarkEditorPerProjectUserSettings* Settings = GetDefault<UWorldBookmarkEditorPerProjectUserSettings>();
+	if (Settings->bEnableDefaultBookmarks)
+	{
+		// Skip loading the default bookmark if the map change was actually caused by loading a bookmark
+		if (!UEditorStateSubsystem::Get()->IsRestoringEditorState())
+		{
+			if (IsDefaultBookmarkValid(InWorld->GetWorldSettings()))
+			{
+				// At this point, make sure we are really dealing with a World Bookmark. Otherwise, do nothing.
+				if (UWorldBookmark* WorldBookmark = Cast<UWorldBookmark>(InWorld->GetWorldSettings()->GetDefaultBookmark()))
+				{
+					WorldBookmark->Load();
+				}
+			}
+			else
+			{					
+				ShowInvalidDefaultBookmarkNotification(LOCTEXT("DefaultBookmarkIncorrectWorld_OnMapChanged", "Default bookmark not applied"));
+			}
+		}
+	}
+}
+
+void FWorldBookmarkModule::OnDefaultBookmarkChanged(AWorldSettings* InWorldSettings)
+{
+	if (!IsDefaultBookmarkValid(InWorldSettings))
+	{
+		ShowInvalidDefaultBookmarkNotification(LOCTEXT("DefaultBookmarkIncorrectWorld_OnDefaultBookmarkChanged", "Invalid default bookmark"));
+		InWorldSettings->SetDefaultBookmark(nullptr);
+	}
+}
+
+bool FWorldBookmarkModule::IsDefaultBookmarkValid(const AWorldSettings* InWorldSettings) const
+{
+	UWorldBookmark* WorldBookmark = InWorldSettings ? Cast<UWorldBookmark>(InWorldSettings->GetDefaultBookmark()) : nullptr;
+	if (!WorldBookmark)
+	{
+		return true; // Not specifying any bookmark (or, possibly, another type of bookmark) is valid
+	}
+
+	// Validate that the world bookmark's world is our actual world
+	TSoftObjectPtr<UWorld> BookmarkWorld = WorldBookmark->GetEditorState<UWorldEditorState>()->GetStateWorld();
+	return BookmarkWorld == InWorldSettings->GetWorld();
+}
+
+void FWorldBookmarkModule::ShowInvalidDefaultBookmarkNotification(const FText& InNotificationTitle) const
+{
+	// Show toast.
+	FNotificationInfo Info(InNotificationTitle);
+	Info.SubText = LOCTEXT("DefaultBookmarkIncorrectWorld_NotificationSubText", "The default bookmark is referencing another world");
+	Info.ExpireDuration = 5.0f;
+	Info.bFireAndForget = true;
+	Info.Image = FCoreStyle::Get().GetBrush(TEXT("MessageLog.Warning"));
+	Info.HyperlinkText = LOCTEXT("DefaultBookmarkIncorrectWorld_ShowWorldSettings", "Show World Settings");
+	Info.Hyperlink = FSimpleDelegate::CreateLambda([=]()
+	{
+		// Open the World Settings tab
+		FLevelEditorModule& LevelEditorModule = FModuleManager::LoadModuleChecked<FLevelEditorModule>("LevelEditor");
+		LevelEditorModule.GetLevelEditorTabManager()->TryInvokeTab(FName("WorldSettingsTab"));
+	});
+
+	FSlateNotificationManager::Get().AddNotification(Info);
+}
+
 #undef LOCTEXT_NAMESPACE
diff --git a/Engine/Source/Editor/WorldBookmark/Private/WorldBookmark/WorldBookmarkModule.h b/Engine/Source/Editor/WorldBookmark/Private/WorldBookmark/WorldBookmarkModule.h
--- a/Engine/Source/Editor/WorldBookmark/Private/WorldBookmark/WorldBookmarkModule.h
+++ b/Engine/Source/Editor/WorldBookmark/Private/WorldBookmark/WorldBookmarkModule.h
@@ -5,6 +5,7 @@
 #include "Delegates/IDelegateInstance.h"
 #include "Modules/ModuleInterface.h"
 #include "Templates/SharedPointer.h"
+#include "UnrealEdMisc.h"
 
 class FSpawnTabArgs;
 class FTabManager;
@@ -33,8 +34,15 @@ private:
 	TSharedRef<SWidget> CreateWorldBookmarkBrowser();
 
 	void OnAddExtraObjectsToDelete(const TArray<UObject*>& InObjectsToDelete, TSet<UObject*>& OutSecondaryObjects);
+	void OnMapChanged(UWorld* InWorld, EMapChangeType InMapChangeType);
+	void OnDefaultBookmarkChanged(AWorldSettings* InWorldSettings);
+
+	bool IsDefaultBookmarkValid(const AWorldSettings* InWorldSettings) const;
+	void ShowInvalidDefaultBookmarkNotification(const FText& InNotificationTitle) const;
 
 	FDelegateHandle OnAddExtraObjectsToDeleteDelegateHandle;
+	FDelegateHandle OnMapChangedHandle;
+	FDelegateHandle OnDefaultBookmarkChangedHandle;
 
 	TArray<FName> ClassesToUnregisterOnShutdown;
 	TArray<FName> StructsToUnregisterOnShutdown;
diff --git a/Engine/Source/Editor/WorldBookmark/Public/WorldBookmark/WorldBookmark.h b/Engine/Source/Editor/WorldBookmark/Public/WorldBookmark/WorldBookmark.h
--- a/Engine/Source/Editor/WorldBookmark/Public/WorldBookmark/WorldBookmark.h
+++ b/Engine/Source/Editor/WorldBookmark/Public/WorldBookmark/WorldBookmark.h
@@ -5,6 +5,7 @@
 #include "UObject/Interface.h"
 #include "UObject/ObjectSaveContext.h"
 #include "EditorState/EditorStateCollection.h"
+#include "Engine/BookmarkBase.h"
 #include "WorldBookmark.generated.h"
 
 /**
@@ -28,7 +29,7 @@ struct FWorldBookmarkCategory
  * World Bookmarks are assets that stores the state of the editor world.
  */
 UCLASS(MinimalAPI, Config = UserWorldBookmarks, PerObjectConfig)
-class UWorldBookmark : public UObject
+class UWorldBookmark : public UBookmarkBase
 {
 	GENERATED_UCLASS_BODY()
 
diff --git a/Engine/Source/Editor/WorldBookmark/WorldBookmark.Build.cs b/Engine/Source/Editor/WorldBookmark/WorldBookmark.Build.cs
--- a/Engine/Source/Editor/WorldBookmark/WorldBookmark.Build.cs
+++ b/Engine/Source/Editor/WorldBookmark/WorldBookmark.Build.cs
@@ -15,6 +15,7 @@ public class WorldBookmark : ModuleRules
 				"ContentBrowserData",
 				"Core",
 				"CoreUObject",
+				"DeveloperSettings",
 				"EditorSubsystem",
 				"EditorConfig",
 				"Engine",
diff --git a/Engine/Source/Runtime/Engine/Classes/GameFramework/WorldSettings.h b/Engine/Source/Runtime/Engine/Classes/GameFramework/WorldSettings.h
--- a/Engine/Source/Runtime/Engine/Classes/GameFramework/WorldSettings.h
+++ b/Engine/Source/Runtime/Engine/Classes/GameFramework/WorldSettings.h
@@ -899,6 +899,12 @@ private:
 	UPROPERTY()
 	TSubclassOf<class UBookmarkBase> LastBookmarkClass;
 
+#if WITH_EDITORONLY_DATA
+	// Bookmark that will be applied when the level is loaded in the editor.
+	UPROPERTY(EditAnywhere, Category=Bookmark)
+	TObjectPtr<class UBookmarkBase> DefaultBookmark;
+#endif 
+
 public:
 	ENGINE_API virtual FSoftClassPath GetAISystemClassName() const;
 
@@ -913,6 +919,9 @@ public:
 
 	DECLARE_MULTICAST_DELEGATE_OneParam(FOnNumberOfBookmarksChanged, AWorldSettings*);
 	static ENGINE_API FOnNumberOfBookmarksChanged OnNumberOfBoomarksChanged;
+
+	DECLARE_MULTICAST_DELEGATE_OneParam(FOnDefaultBookmarkChanged, AWorldSettings*);
+	static ENGINE_API FOnDefaultBookmarkChanged OnDefaultBookmarkChanged;
 #endif
 
 	DECLARE_MULTICAST_DELEGATE_OneParam(FOnNaniteSettingsChanged, AWorldSettings*);
@@ -982,6 +991,18 @@ public:
 	 */
 	ENGINE_API void ClearAllBookmarks();
 
+#if WITH_EDITOR
+	/**
+	 * Get the bookmark that will be applied when the map is opened in the editor.
+	 */
+	ENGINE_API class UBookmarkBase* GetDefaultBookmark() const;
+
+	/**
+	 * Set the bookmark that will be applied when the map is opened in the editor.
+	 */
+	ENGINE_API void SetDefaultBookmark(class UBookmarkBase* InDefaultBookmark);
+#endif
+
 #if WITH_EDITORONLY_DATA
 private: //DEPRECATED
 	UPROPERTY()
diff --git a/Engine/Source/Runtime/Engine/Private/WorldSettings.cpp b/Engine/Source/Runtime/Engine/Private/WorldSettings.cpp
--- a/Engine/Source/Runtime/Engine/Private/WorldSettings.cpp
+++ b/Engine/Source/Runtime/Engine/Private/WorldSettings.cpp
@@ -51,6 +51,7 @@ ENGINE_API float GNewWorldToMetersScale = 0.0f;
 #if WITH_EDITOR
 AWorldSettings::FOnBookmarkClassChanged AWorldSettings::OnBookmarkClassChanged;
 AWorldSettings::FOnNumberOfBookmarksChanged AWorldSettings::OnNumberOfBoomarksChanged;
+AWorldSettings::FOnDefaultBookmarkChanged AWorldSettings::OnDefaultBookmarkChanged;
 
 namespace WorldSettingsUtils
 {
@@ -901,6 +902,10 @@ void AWorldSettings::InternalPostPropertyChanged(FName PropertyName)
 	{
 		UpdateBookmarkClass();
 	}
+	else if (PropertyName == GET_MEMBER_NAME_CHECKED(AWorldSettings, DefaultBookmark))
+	{
+		OnDefaultBookmarkChanged.Broadcast(this);
+	}
 
 	if (GetWorld() != nullptr && GetWorld()->PersistentLevel && GetWorld()->PersistentLevel->GetWorldSettings() == this)
 	{
@@ -1039,6 +1044,20 @@ void AWorldSettings::ClearAllBookmarks()
 	}
 }
 
+#if WITH_EDITOR
+
+UBookmarkBase* AWorldSettings::GetDefaultBookmark() const
+{
+	return DefaultBookmark;
+}
+
+void AWorldSettings::SetDefaultBookmark(UBookmarkBase* InDefaultBookmark)
+{
+	DefaultBookmark = InDefaultBookmark;
+}
+
+#endif
+
 void AWorldSettings::AdjustNumberOfBookmarks()
 {
 	if (MaxNumberOfBookmarks < 0)
