World Bookmark Browser - Fixed alignment issue seen in list view mode

diff --git a/Engine/Source/Editor/WorldBookmark/Private/WorldBookmark/Browser/BookmarkTableRow.cpp b/Engine/Source/Editor/WorldBookmark/Private/WorldBookmark/Browser/BookmarkTableRow.cpp
--- a/Engine/Source/Editor/WorldBookmark/Private/WorldBookmark/Browser/BookmarkTableRow.cpp
+++ b/Engine/Source/Editor/WorldBookmark/Private/WorldBookmark/Browser/BookmarkTableRow.cpp
@@ -31,27 +31,21 @@ void FWorldBookmarkTableRow::Construct(const FArguments& InArgs, const TSharedRe
 
 TSharedRef<SWidget> FWorldBookmarkTableRow::GenerateWidgetForColumn(const FName& ColumnName)
 {
-	if (FWorldBookmarkTreeItemPtr WorldBookmarkTreeItem = GetTreeItem())
+	if (UWorldBookmark* WorldBookmark = GetBookmark())
 	{
 		if (ColumnName == Columns::Favorite.Id)
 		{
-			if (UWorldBookmark* WorldBookmark = Cast<UWorldBookmark>(WorldBookmarkTreeItem->BookmarkAsset.GetAsset()))
+			if (WorldBookmark->GetIsUserFavorite())
 			{
-				if (WorldBookmark->GetIsUserFavorite())
-				{
-					return CreateIconWidget(GetFavoriteWorldBookmarkIcon(true));
-				}
+				return CreateIconWidget(GetFavoriteWorldBookmarkIcon(true));
 			}
 		}
 		else if (ColumnName == Columns::RecentlyUsed.Id)
 		{
-			if (UWorldBookmark* WorldBookmark = Cast<UWorldBookmark>(WorldBookmarkTreeItem->BookmarkAsset.GetAsset()))
+			FTimespan LastLoaded = FDateTime::UtcNow() - WorldBookmark->GetUserLastLoadedTimeStampUTC();
+			if (LastLoaded < FTimespan::FromDays(14))
 			{
-				FTimespan LastLoaded = FDateTime::UtcNow() - WorldBookmark->GetUserLastLoadedTimeStampUTC();
-				if (LastLoaded < FTimespan::FromDays(14))
-				{
-					return CreateIconWidget(GetRecentlyUsedWorldBookmarkIcon());
-				}
+				return CreateIconWidget(GetRecentlyUsedWorldBookmarkIcon());
 			}
 		}
 		else if (ColumnName == Columns::Label.Id)
@@ -99,43 +93,69 @@ TSharedRef<SWidget> FWorldBookmarkTableRow::GenerateLabelForTreeView()
 
 TSharedRef<SWidget> FWorldBookmarkTableRow::GenerateLabelForListView()
 {
-	return SNew(SVerticalBox)
-		+ SVerticalBox::Slot()
-		.Padding(10, 5, 0, 0)
-		.AutoHeight()
-		[
-			CreateEditableLabelWidget(SharedThis(this))
-		]
-		+ SVerticalBox::Slot()
-		.Padding(10, 0, 0, 5)
-		.AutoHeight()
+	return SNew(SBox)
+		.VAlign(EVerticalAlignment::VAlign_Center)
+		.HeightOverride(35)
 		[
-			SNew(STextBlock)
-				.ColorAndOpacity(FSlateColor::UseSubduedForeground())
-				.Text(this, &FWorldBookmarkTableRow::GetLastAccessedText)
+			SNew(SVerticalBox)
+				+ SVerticalBox::Slot()
+				.Padding(10, 0, 0, 0)
+				.AutoHeight()
+				.VAlign(EVerticalAlignment::VAlign_Center)
+				[
+					CreateEditableLabelWidget(SharedThis(this))
+				]
+				+ SVerticalBox::Slot()
+				.Padding(10, 0, 0, 0)
+				[
+					SNew(STextBlock)
+						.ColorAndOpacity(FSlateColor::UseSubduedForeground())
+						.Text(this, &FWorldBookmarkTableRow::GetLastAccessedText)
+						.Visibility(this, &FWorldBookmarkTableRow::GetLastAccessedTextVisibility)
+				]
 		];
 }
 
-FText FWorldBookmarkTableRow::GetLastAccessedText() const
+UWorldBookmark* FWorldBookmarkTableRow::GetBookmark() const
 {
 	if (FWorldBookmarkTreeItemPtr WorldBookmarkTreeItem = GetTreeItem())
 	{
-		if (UWorldBookmark* WorldBookmark = Cast<UWorldBookmark>(WorldBookmarkTreeItem->BookmarkAsset.GetAsset()))
+		return Cast<UWorldBookmark>(WorldBookmarkTreeItem->BookmarkAsset.GetAsset());
+	}
+
+	return nullptr;
+}
+
+FText FWorldBookmarkTableRow::GetLastAccessedText() const
+{
+	if (UWorldBookmark* WorldBookmark = GetBookmark())
+	{
+		FDateTime LastLoadedUTC = WorldBookmark->GetUserLastLoadedTimeStampUTC();
+		if (LastLoadedUTC != FDateTime::MinValue())
 		{
-			FDateTime LastLoadedUTC = WorldBookmark->GetUserLastLoadedTimeStampUTC();
+			FTimespan UTCOffset = FDateTime::UtcNow() - FDateTime::Now();
+			uint64 AdjustedTicks = LastLoadedUTC.GetTicks() - UTCOffset.GetTicks();
+			FDateTime LastLoadedLocal = FDateTime(AdjustedTicks);
 
-			if (LastLoadedUTC != FDateTime::MinValue())
-			{
-				FTimespan UTCOffset = FDateTime::UtcNow() - FDateTime::Now();
-				uint64 AdjustedTicks = LastLoadedUTC.GetTicks() - UTCOffset.GetTicks();
-				FDateTime LastLoadedLocal = FDateTime(AdjustedTicks);
+			return FText::Format(LOCTEXT("LastAccessedDateTime", "Last accessed {0}"), FText::AsDateTime(LastLoadedUTC, TEXT("%Y-%m-%d %H:%M")));
+		}
+	}
+	
+	return FText();
+}
 
-				return FText::Format(LOCTEXT("LastAccessedDateTime", "Last accessed {0}"), FText::AsDateTime(LastLoadedUTC, TEXT("%Y-%m-%d %H:%M")));
-			}
+EVisibility FWorldBookmarkTableRow::GetLastAccessedTextVisibility() const
+{
+	if (UWorldBookmark* WorldBookmark = GetBookmark())
+	{
+		FDateTime LastLoadedUTC = WorldBookmark->GetUserLastLoadedTimeStampUTC();
+		if (LastLoadedUTC != FDateTime::MinValue())
+		{
+			return EVisibility::Visible;
 		}
 	}
 
-	return FText();
+	return EVisibility::Collapsed;
 }
 
 FText FWorldBookmarkTableRow::GetWorldNameText() const
@@ -150,14 +170,11 @@ FText FWorldBookmarkTableRow::GetWorldNameText() const
 
 FLinearColor FWorldBookmarkTableRow::GetCategoryColor() const
 {
-	if (FWorldBookmarkTreeItemPtr WorldBookmarkTreeItem = GetTreeItem())
+	if (UWorldBookmark* WorldBookmark = GetBookmark())
 	{
-		if (UWorldBookmark* WorldBookmark = Cast<UWorldBookmark>(WorldBookmarkTreeItem->BookmarkAsset.GetAsset()))
+		if (WorldBookmark->GetBookmarkCategory().Name != NAME_None)
 		{
-			if (WorldBookmark->GetBookmarkCategory().Name != NAME_None)
-			{
-				return WorldBookmark->GetBookmarkCategory().Color;
-			}
+			return WorldBookmark->GetBookmarkCategory().Color;
 		}
 	}
 
@@ -166,12 +183,9 @@ FLinearColor FWorldBookmarkTableRow::GetCategoryColor() const
 
 FText FWorldBookmarkTableRow::GetCategoryText() const
 {
-	if (FWorldBookmarkTreeItemPtr WorldBookmarkTreeItem = GetTreeItem())
+	if (UWorldBookmark* WorldBookmark = GetBookmark())
 	{
-		if (UWorldBookmark* WorldBookmark = Cast<UWorldBookmark>(WorldBookmarkTreeItem->BookmarkAsset.GetAsset()))
-		{
-			return FText::FromName(WorldBookmark->GetBookmarkCategory().Name);
-		}
+		return FText::FromName(WorldBookmark->GetBookmarkCategory().Name);
 	}
 
 	return FText();
diff --git a/Engine/Source/Editor/WorldBookmark/Private/WorldBookmark/Browser/BookmarkTableRow.h b/Engine/Source/Editor/WorldBookmark/Private/WorldBookmark/Browser/BookmarkTableRow.h
--- a/Engine/Source/Editor/WorldBookmark/Private/WorldBookmark/Browser/BookmarkTableRow.h
+++ b/Engine/Source/Editor/WorldBookmark/Private/WorldBookmark/Browser/BookmarkTableRow.h
@@ -7,6 +7,8 @@
 #include "WorldBookmark/Browser/BookmarkTreeItem.h"
 #include "WorldBookmark/Browser/CommonTableRow.h"
 
+class UWorldBookmark;
+
 namespace UE::WorldBookmark::Browser
 {
 
@@ -21,12 +23,16 @@ public:
 	TSharedRef<SWidget> GenerateLabelForListView();
 
 	FText GetLastAccessedText() const;
+	EVisibility GetLastAccessedTextVisibility() const;
+
 	FText GetWorldNameText() const;
 
 	FLinearColor GetCategoryColor() const;
 	FText GetCategoryText() const;
 
 	FSlateColor GetIconColor() const;
+
+	UWorldBookmark* GetBookmark() const;
 };
 
 }
\ No newline at end of file
