World Bookmark - Added per project categories setting
* Previously, categories were dynamically populated, which prevented categories from being renamed globally, not without checking out every bookmark using it.
* Bookmarks now store a simple category guid, and categories are actually stored in the project's settings.

diff --git a/Engine/Source/Editor/WorldBookmark/Private/WorldBookmark/Browser/CommonTableRow.h b/Engine/Source/Editor/WorldBookmark/Private/WorldBookmark/Browser/CommonTableRow.h
--- a/Engine/Source/Editor/WorldBookmark/Private/WorldBookmark/Browser/CommonTableRow.h
+++ b/Engine/Source/Editor/WorldBookmark/Private/WorldBookmark/Browser/CommonTableRow.h
@@ -220,7 +220,7 @@ TSharedRef<SWidget> CreateTreeLabelWidget(TSharedRef<TTableRow> TableRow)
 					SNew(SImage)
 					.Image_Lambda([TableRow] { return TableRow->GetIcon(); })
 					.ToolTipText(TableRow->GetTooltipText())
-					.ColorAndOpacity(TableRow->GetIconColor())
+					.ColorAndOpacity_Lambda([TableRow] { return TableRow->GetIconColor(); })
 				]
 			]
 
diff --git a/Engine/Source/Editor/WorldBookmark/Private/WorldBookmark/WorldBookmark.cpp b/Engine/Source/Editor/WorldBookmark/Private/WorldBookmark/WorldBookmark.cpp
--- a/Engine/Source/Editor/WorldBookmark/Private/WorldBookmark/WorldBookmark.cpp
+++ b/Engine/Source/Editor/WorldBookmark/Private/WorldBookmark/WorldBookmark.cpp
@@ -13,6 +13,7 @@
 #include "ScopedTransaction.h"
 #include "UObject/AssetRegistryTagsContext.h"
 #include "Widgets/Notifications/SNotificationList.h"
+#include "WorldBookmark/WorldBookmarkEditorSettings.h"
 
 #include UE_INLINE_GENERATED_CPP_BY_NAME(WorldBookmark)
 
@@ -94,19 +95,6 @@ FSoftObjectPath UWorldBookmark::GetWorldFromAssetData(const FAssetData& InAssetD
 	return WorldPath;
 }
 
-FWorldBookmarkCategory UWorldBookmark::GetCategoryFromAssetData(const FAssetData& InAssetData)
-{
-	FWorldBookmarkCategory Category;
-
-	FString CategoryImportText;
-	if (InAssetData.GetTagValue(GetCategoryAssetTag(), CategoryImportText))
-	{
-		FWorldBookmarkCategory::StaticStruct()->ImportText(*CategoryImportText, &Category, nullptr, PPF_None, GLog, FWorldBookmarkCategory::StaticStruct()->GetName());
-	}
-
-	return Category;
-}
-
 void UWorldBookmark::GetAssetRegistryTags(FAssetRegistryTagsContext Context) const
 {
 	Super::GetAssetRegistryTags(Context);
@@ -121,13 +109,6 @@ void UWorldBookmark::GetAssetRegistryTags(FAssetRegistryTagsContext Context) con
 
 			Context.AddTag(FAssetRegistryTag(GetWorldNameAssetTag(), WorldPath.ToString(), FAssetRegistryTag::TT_Hidden));
 		}
-
-		if (!Category.Name.IsNone())
-		{
-			FString CategoryExportText;
-			FWorldBookmarkCategory::StaticStruct()->ExportText(CategoryExportText, &Category, nullptr, nullptr, PPF_None, nullptr);
-			Context.AddTag(FAssetRegistryTag(GetCategoryAssetTag(), CategoryExportText, FAssetRegistryTag::TT_Hidden));
-		}
 	}
 }
 
@@ -311,7 +292,21 @@ void UWorldBookmark::SetUserLastLoadedTimeStampUTC(const FDateTime& InLastLoaded
 
 const FWorldBookmarkCategory& UWorldBookmark::GetBookmarkCategory() const
 {
-	return Category;
+	return UWorldBookmarkEditorSettings::GetCategory(CategoryGuid);
+}
+
+const FWorldBookmarkCategory FWorldBookmarkCategory::None = FWorldBookmarkCategory(NAME_None, FColor::Black);
+
+FWorldBookmarkCategory::FWorldBookmarkCategory()
+	: FWorldBookmarkCategory(TEXT("New Category"), FColor::MakeRandomColor())
+{
+}
+
+FWorldBookmarkCategory::FWorldBookmarkCategory(FName InName, FColor InColor)
+	: Name(InName)
+	, Color(InColor)
+	, Guid(InName == NAME_None ? FGuid() : FGuid::NewGuid())
+{
 }
 
 bool FWorldBookmarkCategory::operator<(const FWorldBookmarkCategory& Other) const
diff --git a/Engine/Source/Editor/WorldBookmark/Private/WorldBookmark/WorldBookmarkDetailsCustomization.cpp b/Engine/Source/Editor/WorldBookmark/Private/WorldBookmark/WorldBookmarkDetailsCustomization.cpp
--- a/Engine/Source/Editor/WorldBookmark/Private/WorldBookmark/WorldBookmarkDetailsCustomization.cpp
+++ b/Engine/Source/Editor/WorldBookmark/Private/WorldBookmark/WorldBookmarkDetailsCustomization.cpp
@@ -2,6 +2,7 @@
 
 #include "WorldBookmark/WorldBookmarkDetailsCustomization.h"
 #include "WorldBookmark/WorldBookmark.h"
+#include "WorldBookmark/WorldBookmarkEditorSettings.h"
 
 #include "AssetRegistry/AssetData.h"
 #include "AssetRegistry/IAssetRegistry.h"
@@ -13,6 +14,7 @@
 #include "IDetailChildrenBuilder.h"
 #include "IStructureDetailsView.h"
 #include "PropertyHandle.h"
+#include "ScopedTransaction.h"
 #include "SPrimaryButton.h"
 #include "Widgets/Colors/SColorBlock.h"
 #include "Widgets/Input/SComboBox.h"
@@ -21,15 +23,21 @@
 
 #define LOCTEXT_NAMESPACE "WorldBookmarkDetailsCustomization"
 
+static const FName NAME_NewCategory("New Category...");
+
 FWorldBookmarkDetailsCustomization::FWorldBookmarkDetailsCustomization()
 	: CachedDetailBuilder()
 {
 	GEditor->RegisterForUndo(this);
+
+	OnWorldBookmarkEditorSettingsChangedHandle = UWorldBookmarkEditorSettings::OnSettingsChanged().AddRaw(this, &FWorldBookmarkDetailsCustomization::OnWorldBookmarkSettingsChanged);
 }
 
 FWorldBookmarkDetailsCustomization::~FWorldBookmarkDetailsCustomization()
 {
 	GEditor->UnregisterForUndo(this);
+
+	UWorldBookmarkEditorSettings::OnSettingsChanged().Remove(OnWorldBookmarkEditorSettingsChangedHandle);
 }
 
 TSharedRef<IDetailCustomization> FWorldBookmarkDetailsCustomization::MakeInstance()
@@ -48,13 +56,59 @@ void FWorldBookmarkDetailsCustomization::CustomizeDetails(IDetailLayoutBuilder&
 		return;
 	}
 
-	UWorldBookmark* WorldBookmark = Cast<UWorldBookmark>(ObjectsBeingCustomized[0].Get());
+	WorldBookmark = Cast<UWorldBookmark>(ObjectsBeingCustomized[0].Get());
 	if (WorldBookmark == nullptr)
 	{
 		return;
 	}
 
-	// Sort categories alphabetically
+	RefreshBookmarkCategoriesList();
+
+	// Customize the FGuid property
+	TSharedRef<IPropertyHandle> GuidPropertyHandle = DetailBuilder.GetProperty(GET_MEMBER_NAME_CHECKED(UWorldBookmark, CategoryGuid));
+	IDetailPropertyRow* GuidPropertyRow = DetailBuilder.EditDefaultProperty(GuidPropertyHandle);
+	
+	GuidPropertyRow->CustomWidget()
+		.NameContent()
+		[
+			GuidPropertyHandle->CreatePropertyNameWidget() // Use the default name widget
+		]
+		.ValueContent()
+		.MinDesiredWidth(200)
+		[
+			SAssignNew(CategoriesComboBox, SComboBox<TSharedPtr<FWorldBookmarkCategory>>)
+				.OptionsSource(&KnownCategories)
+				.OnGenerateWidget(this, &FWorldBookmarkDetailsCustomization::MakeCategoryComboWidget)
+				.OnSelectionChanged(this, &FWorldBookmarkDetailsCustomization::OnCategoryChanged)
+				[
+					SNew(SHorizontalBox)
+
+						+ SHorizontalBox::Slot()
+						.MaxWidth(20)
+						[
+							SNew(SBox)
+								.Visibility(this, &FWorldBookmarkDetailsCustomization::GetCategoryColorVisibility)
+								.Padding(0, 0, 4, 0)
+								[
+									SNew(SColorBlock)
+										.Color(this, &FWorldBookmarkDetailsCustomization::GetCategoryColor)
+										.CornerRadius(FVector4(4.0f, 4.0f, 4.0f, 4.0f))
+										.Size(FVector2D(20, 16))
+								]
+						]
+
+						+ SHorizontalBox::Slot()
+						.VAlign(VAlign_Center)
+						.AutoWidth()
+						[
+							SNew(STextBlock)
+								.Text(this, &FWorldBookmarkDetailsCustomization::GetCategoryText)
+								.Font(IDetailLayoutBuilder::GetDetailFont())
+						]
+				]
+		];
+
+	// Sort state categories alphabetically
 	TArray<UEditorState*> SortedEditorStateObjects = WorldBookmark->EditorState.GetStates();
 	SortedEditorStateObjects.Sort([](UEditorState& A, UEditorState& B)
 	{
@@ -68,173 +122,72 @@ void FWorldBookmarkDetailsCustomization::CustomizeDetails(IDetailLayoutBuilder&
 	}
 }
 
+void FWorldBookmarkDetailsCustomization::OnWorldBookmarkSettingsChanged(UObject* InSettingsObj, FPropertyChangedEvent& InEvent)
+{
+	RefreshCustomDetail();
+}
+
 void FWorldBookmarkDetailsCustomization::RefreshCustomDetail()
 {
 	if (CachedDetailBuilder)
 	{
+		RefreshBookmarkCategoriesList();
+
 		CachedDetailBuilder->ForceRefreshDetails();
 	}
 }
 
-void FWorldBookmarkDetailsCustomization::PostUndo(bool bSuccess)
+void FWorldBookmarkDetailsCustomization::RefreshBookmarkCategoriesList()
 {
-	// Refresh the UI
-	RefreshCustomDetail();
-}
+	KnownCategories.Reset();
 
-void FWorldBookmarkDetailsCustomization::PostRedo(bool bSuccess)
-{
-	// Refresh the UI
-	RefreshCustomDetail();
-}
-
-TSharedRef<IPropertyTypeCustomization> FWorldBookmarkCategoryCustomization::MakeInstance()
-{
-	return MakeShareable(new FWorldBookmarkCategoryCustomization);
-}
-
-
-const FName FWorldBookmarkCategoryCustomization::NAME_NewCategory("New Category...");
-
-FWorldBookmarkCategoryCustomization::FWorldBookmarkCategoryCustomization()
-{	
-	IAssetRegistry& AssetRegistry = IAssetRegistry::GetChecked();
-	TArray<FAssetData> AssetsData;
-	const bool bSearchSubClasses = true;
-	AssetRegistry.GetAssetsByClass(UWorldBookmark::StaticClass()->GetClassPathName(), AssetsData, bSearchSubClasses);
-		
-	// Add default "categories" to the list of items that will be shown in the Category combo box
-	AddCategory(FWorldBookmarkCategory{ NAME_None, FColor::Black });
-	AddCategory(FWorldBookmarkCategory{ NAME_NewCategory, FColor::Black });
-
-	for (const FAssetData& AssetData : AssetsData)
+	// Add "None"
+	KnownCategories.Emplace(MakeShared<FWorldBookmarkCategory>(FWorldBookmarkCategory{ NAME_None, FColor::Black }));
+	
+	// Add categories found in the project's settings
+	for (const FWorldBookmarkCategory& Category : UWorldBookmarkEditorSettings::GetCategories())
 	{
-		FWorldBookmarkCategory Category = UWorldBookmark::GetCategoryFromAssetData(AssetData);
-		AddCategory(Category);
+		KnownCategories.Emplace(MakeShared<FWorldBookmarkCategory>(Category));
 	}
 
-	SortCategories();
-}
-
-void FWorldBookmarkCategoryCustomization::AddCategory(const FWorldBookmarkCategory& InCategory)
-{
-	bool bFoundExistingCategory = KnownCategories.ContainsByPredicate([&InCategory](const TSharedPtr<FWorldBookmarkCategory>& ExistingCategory)
-	{
-		return ExistingCategory->Name == InCategory.Name;
-	});
+	// Add "New Category..."
+	KnownCategories.Emplace(MakeShared<FWorldBookmarkCategory>(FWorldBookmarkCategory{ NAME_NewCategory, FColor::Black }));
 
-	if (!bFoundExistingCategory)
+	// Refresh the combo box
+	if (CategoriesComboBox)
 	{
-		KnownCategories.Emplace(MakeShared<FWorldBookmarkCategory>(InCategory));
+		CategoriesComboBox->RefreshOptions();
 	}
 }
 
-void FWorldBookmarkCategoryCustomization::SortCategories()
+void FWorldBookmarkDetailsCustomization::PostUndo(bool bSuccess)
 {
-	KnownCategories.Sort([](const TSharedPtr<FWorldBookmarkCategory>& LHS, const TSharedPtr<FWorldBookmarkCategory>& RHS)
-	{
-		if (LHS->Name == NAME_None || RHS->Name == NAME_None)
-		{
-			return LHS->Name == NAME_None;
-		}
-
-		if (LHS->Name == NAME_NewCategory || RHS->Name == NAME_NewCategory)
-		{
-			return LHS->Name == NAME_NewCategory;
-		}
-
-		return LHS->Name.LexicalLess(RHS->Name);
-	});
+	// Refresh the UI
+	RefreshCustomDetail();
 }
 
-void FWorldBookmarkCategoryCustomization::CustomizeHeader(TSharedRef<IPropertyHandle> InStructPropertyHandle, FDetailWidgetRow& InHeaderRow, IPropertyTypeCustomizationUtils& InCustomizationUtils)
+void FWorldBookmarkDetailsCustomization::PostRedo(bool bSuccess)
 {
-	StructPropertyHandle = InStructPropertyHandle.ToSharedPtr();
-
-	TArray<void*> RawData;
-	InStructPropertyHandle->AccessRawData(RawData);
-	FWorldBookmarkCategory* CurrentCategory = reinterpret_cast<FWorldBookmarkCategory*>(RawData[0]);
-
-	TSharedPtr<FWorldBookmarkCategory>* CurrentCategoryPtr = KnownCategories.FindByPredicate([CurrentCategory](const TSharedPtr<FWorldBookmarkCategory>& Item)
-	{
-		return Item->Name == CurrentCategory->Name;
-	});
-
-	// Create the collapsed view with Name and Color
-	InHeaderRow.NameContent()
-		[
-			StructPropertyHandle->CreatePropertyNameWidget()
-		]
-		.ValueContent()
-		.MinDesiredWidth(200.0f)
-		[
-			SNew(SComboBox<TSharedPtr<FWorldBookmarkCategory>>)
-				.OptionsSource(&KnownCategories)
-				.OnGenerateWidget(this, &FWorldBookmarkCategoryCustomization::MakeCategoryComboWidget)
-				.OnSelectionChanged(this, &FWorldBookmarkCategoryCustomization::OnCategoryChanged)
-				.InitiallySelectedItem(CurrentCategoryPtr ? *CurrentCategoryPtr : TSharedPtr<FWorldBookmarkCategory>())
-				[
-					SNew(SHorizontalBox)
-
-					+ SHorizontalBox::Slot()
-					.MaxWidth(20)
-					[
-						SNew(SBox)
-						.Visibility(this, &FWorldBookmarkCategoryCustomization::GetSelectedCategoryColorVisibility)
-						.Padding(0, 0, 4, 0)						
-						[
-							SNew(SColorBlock)
-							.Color(this, &FWorldBookmarkCategoryCustomization::GetSelectedCategoryColor)
-							.CornerRadius(FVector4(4.0f, 4.0f, 4.0f, 4.0f))
-							.Size(FVector2D(20, 16))
-						]
-					]
-
-					+ SHorizontalBox::Slot()
-					.VAlign(VAlign_Center)
-					.AutoWidth()
-					[
-						SNew(STextBlock)
-						.Text(this, &FWorldBookmarkCategoryCustomization::GetSelectedCategoryText)
-					]
-				]
-		];
+	// Refresh the UI
+	RefreshCustomDetail();
 }
 
-FLinearColor FWorldBookmarkCategoryCustomization::GetSelectedCategoryColor() const
+FLinearColor FWorldBookmarkDetailsCustomization::GetCategoryColor() const
 {
-	TArray<void*> RawData;
-	StructPropertyHandle->AccessRawData(RawData);
-
-	bool bShouldAssign = true;
-	FWorldBookmarkCategory* WorldBookmarkCategory = reinterpret_cast<FWorldBookmarkCategory*>(RawData[0]);
-	return WorldBookmarkCategory->Color;
+	return WorldBookmark->GetBookmarkCategory().Color;
 }
 
-EVisibility FWorldBookmarkCategoryCustomization::GetSelectedCategoryColorVisibility() const
+EVisibility FWorldBookmarkDetailsCustomization::GetCategoryColorVisibility() const
 {
-	TArray<void*> RawData;
-	StructPropertyHandle->AccessRawData(RawData);
-
-	bool bShouldAssign = true;
-	FWorldBookmarkCategory* WorldBookmarkCategory = reinterpret_cast<FWorldBookmarkCategory*>(RawData[0]);
-
-	const bool bShowColorBlock = (WorldBookmarkCategory->Name != NAME_None) && (WorldBookmarkCategory->Name != NAME_NewCategory);
-	return bShowColorBlock ? EVisibility::Visible : EVisibility::Collapsed;
+	return WorldBookmark->GetBookmarkCategory().Name.IsNone() ? EVisibility::Collapsed : EVisibility::Visible;
 }
 
-FText FWorldBookmarkCategoryCustomization::GetSelectedCategoryText() const
+FText FWorldBookmarkDetailsCustomization::GetCategoryText() const
 {
-	TArray<void*> RawData;
-	StructPropertyHandle->AccessRawData(RawData);
-
-	bool bShouldAssign = true;
-	FWorldBookmarkCategory* WorldBookmarkCategory = reinterpret_cast<FWorldBookmarkCategory*>(RawData[0]);
-
-	return FText::FromName(WorldBookmarkCategory->Name);
+	return FText::FromName(WorldBookmark->GetBookmarkCategory().Name);
 }
 
-TSharedRef<SWidget> FWorldBookmarkCategoryCustomization::MakeCategoryComboWidget(TSharedPtr<FWorldBookmarkCategory> InItem)
+TSharedRef<SWidget> FWorldBookmarkDetailsCustomization::MakeCategoryComboWidget(TSharedPtr<FWorldBookmarkCategory> InItem)
 {
 	TSharedRef<SHorizontalBox> ItemWidget = SNew(SHorizontalBox);
 
@@ -242,56 +195,62 @@ TSharedRef<SWidget> FWorldBookmarkCategoryCustomization::MakeCategoryComboWidget
 	if (bShowColorBlock)
 	{
 		ItemWidget->AddSlot()
-		.MaxWidth(16)
-		[
-			SNew(SColorBlock)
-				.Color(InItem->Color)
-				.CornerRadius(FVector4(4.0f, 4.0f, 4.0f, 4.0f))
-		];
+			.MaxWidth(16)
+			[
+				SNew(SColorBlock)
+					.Color(InItem->Color)
+					.CornerRadius(FVector4(4.0f, 4.0f, 4.0f, 4.0f))
+			];
 	}
 
 	ItemWidget->AddSlot()
-	.VAlign(VAlign_Center)
-	.AutoWidth()
-	.Padding(bShowColorBlock ? 4 : 0, 0, 0, 0)
-	[
-		SNew(STextBlock)
-			.Text(FText::FromName(InItem->Name))
-	];
+		.VAlign(VAlign_Center)
+		.AutoWidth()
+		.Padding(bShowColorBlock ? 4 : 0, 0, 0, 0)
+		[
+			SNew(STextBlock)
+				.Text(FText::FromName(InItem->Name))
+				.Font(IDetailLayoutBuilder::GetDetailFont())
+		];
 
 	return ItemWidget;
 }
 
-void FWorldBookmarkCategoryCustomization::OnCategoryChanged(TSharedPtr<FWorldBookmarkCategory> NewSelection, ESelectInfo::Type SelectInfo)
+void FWorldBookmarkDetailsCustomization::OnCategoryChanged(TSharedPtr<FWorldBookmarkCategory> NewSelection, ESelectInfo::Type SelectInfo)
 {
 	if (NewSelection.IsValid())
 	{
-		TArray<void*> RawData;
-		StructPropertyHandle->AccessRawData(RawData);
+		FScopedTransaction Transaction(LOCTEXT("ChangeBookmarkCategory", "Change Bookmark Category"));
 
 		bool bShouldAssign = true;
-		FWorldBookmarkCategory* WorldBookmarkCategory = reinterpret_cast<FWorldBookmarkCategory*>(RawData[0]);
 		FWorldBookmarkCategory CategoryToAssign = *NewSelection;
 
 		if (CategoryToAssign.Name == NAME_NewCategory)
 		{
-			bShouldAssign = CreateNewCategory(CategoryToAssign);
-			if (bShouldAssign)
+			bool bCreatedNewCategory = CreateNewCategory(CategoryToAssign);
+			if (bCreatedNewCategory)
 			{
-				AddCategory(CategoryToAssign);
+				UWorldBookmarkEditorSettings::AddCategory(CategoryToAssign);
 			}
+			bShouldAssign = bCreatedNewCategory;
 		}
 
-		if (bShouldAssign)
+		if (bShouldAssign && WorldBookmark->CategoryGuid != CategoryToAssign.Guid)
+		{			
+			WorldBookmark->Modify();
+			WorldBookmark->CategoryGuid = CategoryToAssign.Guid;
+
+			TSharedRef<IPropertyHandle> GuidPropertyHandle = CachedDetailBuilder->GetProperty(GET_MEMBER_NAME_CHECKED(UWorldBookmark, CategoryGuid));
+			GuidPropertyHandle->NotifyPostChange(EPropertyChangeType::ValueSet);
+		}
+		else
 		{
-			StructPropertyHandle->NotifyPreChange();
-			*WorldBookmarkCategory = CategoryToAssign;
-			StructPropertyHandle->NotifyPostChange(EPropertyChangeType::ValueSet);
+			Transaction.Cancel();
 		}
 	}
 }
 
-bool FWorldBookmarkCategoryCustomization::CreateNewCategory(FWorldBookmarkCategory& OutNewCategory)
+bool FWorldBookmarkDetailsCustomization::CreateNewCategory(FWorldBookmarkCategory& OutNewCategory)
 {
 	FPropertyEditorModule& PropertyEditorModule = FModuleManager::GetModuleChecked<FPropertyEditorModule>("PropertyEditor");
 
@@ -307,15 +266,14 @@ bool FWorldBookmarkCategoryCustomization::CreateNewCategory(FWorldBookmarkCatego
 	ViewArgs.bShowObjectLabel = false;
 
 	FWorldBookmarkCategory NewCategoryTemplate;
-	NewCategoryTemplate.Color = FColor::White;
-	TSharedPtr<TStructOnScope<FWorldBookmarkCategory>> Settings;
-	Settings = MakeShared<TStructOnScope<FWorldBookmarkCategory>>();
-	Settings->InitializeAs<FWorldBookmarkCategory>(NewCategoryTemplate);
+	TSharedPtr<TStructOnScope<FWorldBookmarkCategory>> NewCategory;
+	NewCategory = MakeShared<TStructOnScope<FWorldBookmarkCategory>>();
+	NewCategory->InitializeAs<FWorldBookmarkCategory>(NewCategoryTemplate);
 
 	TSharedPtr<IStructureDetailsView> DetailsView = PropertyEditorModule.CreateStructureDetailView(ViewArgs, StructureViewArgs, TSharedPtr<FStructOnScope>());
 
 	// Set provided objects on SDetailsView
-	DetailsView->SetStructureData(Settings);
+	DetailsView->SetStructureData(NewCategory);
 
 	bool bCreatedNewCategory = false;
 
@@ -372,20 +330,83 @@ bool FWorldBookmarkCategoryCustomization::CreateNewCategory(FWorldBookmarkCatego
 	);
 
 	GEditor->EditorAddModalWindow(Window);
-	
+
 	if (bCreatedNewCategory)
 	{
-		OutNewCategory.Name = Settings->Get()->Name;
-		OutNewCategory.Color = Settings->Get()->Color;
+		const bool bValidCategoryName = NewCategory->Get()->Name.ToString().ToLower() != FName(NAME_None).ToString().ToLower() &&
+				                        NewCategory->Get()->Name.ToString().ToLower() != NAME_NewCategory.ToString().ToLower();
+		bCreatedNewCategory = bValidCategoryName;
+		if (bCreatedNewCategory)
+		{
+			OutNewCategory = FWorldBookmarkCategory(NewCategory->Get()->Name, NewCategory->Get()->Color);
+		}		
 	}
-
-	bCreatedNewCategory = OutNewCategory.Name != NAME_None && OutNewCategory.Name.ToString().ToLower() != "none";
+	
 	return bCreatedNewCategory;
 }
 
-void FWorldBookmarkCategoryCustomization::CustomizeChildren(TSharedRef<IPropertyHandle> InStructPropertyHandle, IDetailChildrenBuilder& ChildBuilder, IPropertyTypeCustomizationUtils& CustomizationUtils)
+TSharedRef<IPropertyTypeCustomization> FWorldBookmarkCategoryCustomization::MakeInstance()
+{
+	return MakeShareable(new FWorldBookmarkCategoryCustomization);
+}
+
+void FWorldBookmarkCategoryCustomization::CustomizeHeader(TSharedRef<IPropertyHandle> InStructPropertyHandle, FDetailWidgetRow& InHeaderRow, IPropertyTypeCustomizationUtils& InCustomizationUtils)
+{
+	CachedStructPropertyHandle = InStructPropertyHandle;
+
+	// Create the collapsed view with Name and Color
+	InHeaderRow.NameContent()
+	[
+		InStructPropertyHandle->CreatePropertyNameWidget()
+	]
+	.ValueContent()
+	.MinDesiredWidth(200.0f)
+	[
+		SNew(SHorizontalBox)
+			
+		+ SHorizontalBox::Slot()
+		.VAlign(VAlign_Center)
+		.MaxWidth(20)
+		[
+			SNew(SBox)
+			.Padding(0, 0, 4, 0)						
+			[
+				SNew(SColorBlock)
+				.Color_Lambda([this]() { return GetEditedCategory().Color; })
+				.CornerRadius(FVector4(4.0f, 4.0f, 4.0f, 4.0f))
+				.Size(FVector2D(16, 16))
+			]
+		]
+
+		+ SHorizontalBox::Slot()
+		.VAlign(VAlign_Center)
+		.AutoWidth()
+		[
+			SNew(STextBlock)
+			.Text_Lambda([this]() { return FText::FromName(GetEditedCategory().Name); })
+			.Font(IDetailLayoutBuilder::GetDetailFont())
+		]
+	];
+}
+
+const FWorldBookmarkCategory& FWorldBookmarkCategoryCustomization::GetEditedCategory() const
 {
-	// Do not allow the row to be expanded (never show the category color & name as individual rows)
+	TArray<void*> RawData;
+	CachedStructPropertyHandle->AccessRawData(RawData);
+	FWorldBookmarkCategory* CurrentCategory = reinterpret_cast<FWorldBookmarkCategory*>(RawData[0]);
+	return CurrentCategory ? *CurrentCategory : FWorldBookmarkCategory::None;
+}
+
+// Default behavior, add all children properties
+void FWorldBookmarkCategoryCustomization::CustomizeChildren(TSharedRef<IPropertyHandle> StructPropertyHandle, IDetailChildrenBuilder& ChildBuilder, IPropertyTypeCustomizationUtils& CustomizationUtils)
+{
+	uint32 NumChildren;
+	StructPropertyHandle->GetNumChildren(NumChildren);
+
+	for (uint32 ChildIndex = 0; ChildIndex < NumChildren; ++ChildIndex)
+	{
+		ChildBuilder.AddProperty(StructPropertyHandle->GetChildHandle(ChildIndex).ToSharedRef());
+	}
 }
 
 #undef LOCTEXT_NAMESPACE
diff --git a/Engine/Source/Editor/WorldBookmark/Private/WorldBookmark/WorldBookmarkDetailsCustomization.h b/Engine/Source/Editor/WorldBookmark/Private/WorldBookmark/WorldBookmarkDetailsCustomization.h
--- a/Engine/Source/Editor/WorldBookmark/Private/WorldBookmark/WorldBookmarkDetailsCustomization.h
+++ b/Engine/Source/Editor/WorldBookmark/Private/WorldBookmark/WorldBookmarkDetailsCustomization.h
@@ -5,6 +5,7 @@
 #include "IDetailCustomization.h"
 #include "IPropertyTypeCustomization.h"
 #include "EditorUndoClient.h"
+#include "Widgets/Input/SComboBox.h"
 #include "Widgets/SWidget.h"
 
 class UWorldBookmark;
@@ -32,38 +33,43 @@ protected:
 	void RefreshCustomDetail();
 
 private:
+	TSharedRef<SWidget> MakeCategoryComboWidget(TSharedPtr<FWorldBookmarkCategory> InItem);
+
+	void OnCategoryChanged(TSharedPtr<FWorldBookmarkCategory> NewSelection, ESelectInfo::Type SelectInfo);
+	bool CreateNewCategory(FWorldBookmarkCategory& OutNewCategory);
+
+	FLinearColor GetCategoryColor() const;
+	EVisibility GetCategoryColorVisibility() const;
+	FText GetCategoryText() const;
+
+	void OnWorldBookmarkSettingsChanged(UObject* InSettings, FPropertyChangedEvent& InPropertyChangedEvent);
+	
+	void RefreshBookmarkCategoriesList();
+
 	/** The detail builder for this cusomtomisation. */
 	IDetailLayoutBuilder* CachedDetailBuilder;
+
+	TArray<TSharedPtr<FWorldBookmarkCategory>> KnownCategories;
+	TSharedPtr<SComboBox<TSharedPtr<FWorldBookmarkCategory>>> CategoriesComboBox;
+
+	/** The currently edited bookmark */
+	UWorldBookmark* WorldBookmark;
+
+	FDelegateHandle OnWorldBookmarkEditorSettingsChangedHandle;
 };
 
-/** UI customization for FWorldBookmarkCategory */
+// UI customization for FWorldBookmarkCategory
 class FWorldBookmarkCategoryCustomization : public IPropertyTypeCustomization
 {
 public:
-	/** Factory method to create an instance */
+	// Factory method to create an instance
 	static TSharedRef<IPropertyTypeCustomization> MakeInstance();
 
-	FWorldBookmarkCategoryCustomization();
-
-	/** IPropertyTypeCustomization overrides */
+	// IPropertyTypeCustomization overrides
 	virtual void CustomizeHeader(TSharedRef<IPropertyHandle> StructPropertyHandle, FDetailWidgetRow& HeaderRow, IPropertyTypeCustomizationUtils& CustomizationUtils) override;
 	virtual void CustomizeChildren(TSharedRef<IPropertyHandle> StructPropertyHandle, IDetailChildrenBuilder& ChildBuilder, IPropertyTypeCustomizationUtils& CustomizationUtils) override;
 
-private:
-	TSharedRef<SWidget> MakeCategoryComboWidget(TSharedPtr<FWorldBookmarkCategory> InItem);
-
-	void OnCategoryChanged(TSharedPtr<FWorldBookmarkCategory> NewSelection, ESelectInfo::Type SelectInfo);
-	void AddCategory(const FWorldBookmarkCategory& InCategory);
-	void SortCategories();
-	bool CreateNewCategory(FWorldBookmarkCategory& OutNewCategory);
+	const FWorldBookmarkCategory& GetEditedCategory() const;
 
-	FLinearColor GetSelectedCategoryColor() const;
-	EVisibility GetSelectedCategoryColorVisibility() const;
-	FText GetSelectedCategoryText() const;
-
-	TSharedPtr<IPropertyHandle> StructPropertyHandle;
-
-	static const FName NAME_NewCategory;
-	
-	TArray<TSharedPtr<FWorldBookmarkCategory>> KnownCategories;
-};
\ No newline at end of file
+	TSharedPtr<IPropertyHandle> CachedStructPropertyHandle;
+};
diff --git a/Engine/Source/Editor/WorldBookmark/Private/WorldBookmark/WorldBookmarkEditorSettings.cpp b/Engine/Source/Editor/WorldBookmark/Private/WorldBookmark/WorldBookmarkEditorSettings.cpp
--- a/Engine/Source/Editor/WorldBookmark/Private/WorldBookmark/WorldBookmarkEditorSettings.cpp
+++ b/Engine/Source/Editor/WorldBookmark/Private/WorldBookmark/WorldBookmarkEditorSettings.cpp
@@ -2,6 +2,75 @@
 
 #include "WorldBookmark/WorldBookmarkEditorSettings.h"
 
+UDeveloperSettings::FOnSettingsChanged& UWorldBookmarkEditorSettings::OnSettingsChanged()
+{
+	UWorldBookmarkEditorSettings& Settings = *GetMutableDefault<UWorldBookmarkEditorSettings>();
+	return Settings.OnSettingChanged();
+}
+
+const FWorldBookmarkCategory& UWorldBookmarkEditorSettings::GetCategory(const FGuid& InCategoryGuid)
+{
+	const UWorldBookmarkEditorSettings& Settings = *GetDefault<UWorldBookmarkEditorSettings>();
+	const FWorldBookmarkCategory* FoundCategory = Settings.Categories.FindByPredicate([&InCategoryGuid](const FWorldBookmarkCategory& InCategory)
+	{
+		return InCategory.Guid == InCategoryGuid;
+	});
+
+	if (FoundCategory)
+	{
+		return *FoundCategory;
+	}
+
+	return FWorldBookmarkCategory::None;
+}
+
+void UWorldBookmarkEditorSettings::AddCategory(const FWorldBookmarkCategory& InCategory)
+{
+	UWorldBookmarkEditorSettings* Settings = GetMutableDefault<UWorldBookmarkEditorSettings>();
+	Settings->SetFlags(RF_Transactional); // Settings object is not transactional by default
+
+	FWorldBookmarkCategory* ExistingCategory = Settings->Categories.FindByPredicate([&InCategory](const FWorldBookmarkCategory& Category)
+	{
+		return Category.Guid == InCategory.Guid;
+	});
+
+	bool bSettingsChanged = !ExistingCategory || (ExistingCategory->Name != InCategory.Name && ExistingCategory->Color != InCategory.Color);
+	if (bSettingsChanged)
+	{
+		Settings->Modify();
+
+		if (ExistingCategory)
+		{
+			ExistingCategory->Name = InCategory.Name;
+			ExistingCategory->Color = InCategory.Color;
+		}
+		else
+		{
+			Settings->Categories.Add(InCategory);
+		}
+
+		Settings->PostEditChange();
+	}
+}
+
+void UWorldBookmarkEditorSettings::PostEditChangeProperty(struct FPropertyChangedEvent& PropertyChangedEvent)
+{
+	// Keep categories sorted by names
+	Categories.StableSort([](const FWorldBookmarkCategory& LHS, const FWorldBookmarkCategory& RHS)
+	{
+		return LHS.Name.LexicalLess(RHS.Name);
+	});
+
+	Super::PostEditChangeProperty(PropertyChangedEvent);
+
+	TryUpdateDefaultConfigFile();
+}
+
+const TArray<FWorldBookmarkCategory>& UWorldBookmarkEditorSettings::GetCategories()
+{
+	return GetDefault<UWorldBookmarkEditorSettings>()->Categories;
+}
+
 UWorldBookmarkEditorPerProjectUserSettings::UWorldBookmarkEditorPerProjectUserSettings()
 {
 	bEnableDefaultBookmarks = true;
diff --git a/Engine/Source/Editor/WorldBookmark/Private/WorldBookmark/WorldBookmarkEditorSettings.h b/Engine/Source/Editor/WorldBookmark/Private/WorldBookmark/WorldBookmarkEditorSettings.h
--- a/Engine/Source/Editor/WorldBookmark/Private/WorldBookmark/WorldBookmarkEditorSettings.h
+++ b/Engine/Source/Editor/WorldBookmark/Private/WorldBookmark/WorldBookmarkEditorSettings.h
@@ -3,9 +3,36 @@
 #pragma once
 
 #include "CoreMinimal.h"
+#include "Delegates/Delegate.h"
 #include "Engine/DeveloperSettings.h"
+#include "WorldBookmark/WorldBookmark.h"
 #include "WorldBookmarkEditorSettings.generated.h"
 
+
+UCLASS(config = Editor, DefaultConfig, meta = (DisplayName = "World Bookmark"))
+class UWorldBookmarkEditorSettings : public UDeveloperSettings
+{
+	GENERATED_BODY()
+
+public:
+	DECLARE_MULTICAST_DELEGATE(UWorldBookmarkEditorSettingsChanged)
+
+	static FOnSettingsChanged& OnSettingsChanged();
+
+	static const FWorldBookmarkCategory& GetCategory(const FGuid& InCategoryGuid);
+	static void AddCategory(const FWorldBookmarkCategory& InCategory);
+
+	static const TArray<FWorldBookmarkCategory>& GetCategories();
+
+protected:
+	virtual void PostEditChangeProperty(struct FPropertyChangedEvent& PropertyChangedEvent) override;
+
+private:
+	UPROPERTY(config, EditAnywhere, Category = "Bookmark Categories", meta=(NoElementDuplicate, EditFixedOrder))
+	TArray<FWorldBookmarkCategory> Categories;
+};
+
+
 UCLASS(config = EditorPerProjectUserSettings, meta = (DisplayName = "World Bookmark"))
 class UWorldBookmarkEditorPerProjectUserSettings : public UDeveloperSettings
 {
@@ -15,6 +42,6 @@ public:
 	UWorldBookmarkEditorPerProjectUserSettings();
 
 	/** When enabled, the default bookmark will be applied when loading a level */
-	UPROPERTY(config, EditAnywhere, Category="Default Bookmark")
+	UPROPERTY(config, EditAnywhere, Category="Default Level Bookmark")
 	bool bEnableDefaultBookmarks;
 };
diff --git a/Engine/Source/Editor/WorldBookmark/Private/WorldBookmark/WorldBookmarkModule.cpp b/Engine/Source/Editor/WorldBookmark/Private/WorldBookmark/WorldBookmarkModule.cpp
--- a/Engine/Source/Editor/WorldBookmark/Private/WorldBookmark/WorldBookmarkModule.cpp
+++ b/Engine/Source/Editor/WorldBookmark/Private/WorldBookmark/WorldBookmarkModule.cpp
@@ -44,7 +44,7 @@ void FWorldBookmarkModule::StartupModule()
 		ClassesToUnregisterOnShutdown.Add(UWorldBookmark::StaticClass()->GetFName());
 		PropertyModule.RegisterCustomClassLayout(UWorldBookmark::StaticClass()->GetFName(), FOnGetDetailCustomizationInstance::CreateStatic(&FWorldBookmarkDetailsCustomization::MakeInstance));
 
-		StructsToUnregisterOnShutdown.Add(UWorldBookmark::StaticClass()->GetFName());
+		StructsToUnregisterOnShutdown.Add(FWorldBookmarkCategory::StaticStruct()->GetFName());
 		PropertyModule.RegisterCustomPropertyTypeLayout(FWorldBookmarkCategory::StaticStruct()->GetFName(), FOnGetPropertyTypeCustomizationInstance::CreateStatic(&FWorldBookmarkCategoryCustomization::MakeInstance));
 
 		OnAddExtraObjectsToDeleteDelegateHandle = FEditorDelegates::OnAddExtraObjectsToDelete.AddRaw(this, &FWorldBookmarkModule::OnAddExtraObjectsToDelete);
diff --git a/Engine/Source/Editor/WorldBookmark/Public/WorldBookmark/WorldBookmark.h b/Engine/Source/Editor/WorldBookmark/Public/WorldBookmark/WorldBookmark.h
--- a/Engine/Source/Editor/WorldBookmark/Public/WorldBookmark/WorldBookmark.h
+++ b/Engine/Source/Editor/WorldBookmark/Public/WorldBookmark/WorldBookmark.h
@@ -16,13 +16,21 @@ struct FWorldBookmarkCategory
 {
 	GENERATED_BODY()
 
+	static const FWorldBookmarkCategory None;
+
+	FWorldBookmarkCategory();
+	FWorldBookmarkCategory(FName InName, FColor InColor);
+
 	bool operator<(const FWorldBookmarkCategory& Other) const;
 
 	UPROPERTY(EditAnywhere, Category=Category)
 	FName Name;
 
-	UPROPERTY(EditAnywhere, Category=Category, Meta=(HideAlphaChannel))
+	UPROPERTY(EditAnywhere, Category=Category, Meta=(HideAlphaChannel, DontUpdateWhileEditing))
 	FColor Color;
+
+    UPROPERTY()
+	FGuid Guid;
 };
 
 /**
@@ -78,9 +86,6 @@ class UWorldBookmark : public UBookmarkBase
 	/** Get the world associated with a bookmark's asset data. */
 	static FSoftObjectPath GetWorldFromAssetData(const FAssetData& InAssetData);
 
-	/** Get the bookmark category associated with a bookmark's asset data. */
-	static FWorldBookmarkCategory GetCategoryFromAssetData(const FAssetData& InAssetData);
-
 	template <typename TEditorStateType>
 	typename TEnableIf<TIsDerivedFrom<TEditorStateType, UEditorState>::IsDerived, bool>::Type HasEditorState() const
 	{
@@ -110,8 +115,8 @@ private:
 	UPROPERTY()
 	FEditorStateCollection EditorState;
 
-	UPROPERTY(EditAnywhere, Category=Category)
-	FWorldBookmarkCategory Category;
+	UPROPERTY(EditAnywhere, Category=Category, meta=(DisplayName="Category"))
+	FGuid CategoryGuid;
 
 	// BookmarkGuid is our key to fetch the user settings for this bookmark in the Bookmarks.ini config file.
 	// It will remain unique even if redirectors are created for the world or the bookmark itself.
