World Bookmark - Fixed false positive in member initialization test for FWorldBookmarkCategory
* "FWorldBookmarkCategory is not initialized properly even though its struct probably has a custom default constructor"
* Both Color & Guid members are initialized with random values - mark these properties with "IgnoreForMemberInitializationTest"

diff --git a/Engine/Source/Editor/WorldBookmark/Public/WorldBookmark/WorldBookmark.h b/Engine/Source/Editor/WorldBookmark/Public/WorldBookmark/WorldBookmark.h
--- a/Engine/Source/Editor/WorldBookmark/Public/WorldBookmark/WorldBookmark.h
+++ b/Engine/Source/Editor/WorldBookmark/Public/WorldBookmark/WorldBookmark.h
@@ -26,10 +26,10 @@ struct FWorldBookmarkCategory
 	UPROPERTY(EditAnywhere, Category=Category)
 	FName Name;
 
-	UPROPERTY(EditAnywhere, Category=Category, Meta=(HideAlphaChannel, DontUpdateWhileEditing))
+	UPROPERTY(EditAnywhere, Category=Category, Meta=(HideAlphaChannel, DontUpdateWhileEditing, IgnoreForMemberInitializationTest))
 	FColor Color;
 
-    UPROPERTY()
+    UPROPERTY(Meta=(IgnoreForMemberInitializationTest))
 	FGuid Guid;
 };
 
