World Bookmark - Added the "Home Bookmark" feature
* The home bookmark is a user setting that, when specified, will be loaded/applied when the editor is started without a map argument.
* This will take precedence over the project's default startup map for the editor

diff --git a/Engine/Source/Editor/LevelEditor/Private/LevelEditorCameraEditorState.cpp b/Engine/Source/Editor/LevelEditor/Private/LevelEditorCameraEditorState.cpp
--- a/Engine/Source/Editor/LevelEditor/Private/LevelEditorCameraEditorState.cpp
+++ b/Engine/Source/Editor/LevelEditor/Private/LevelEditorCameraEditorState.cpp
@@ -52,13 +52,32 @@ UEditorState::FOperationResult ULevelEditorCameraEditorState::CaptureState()
 
 UEditorState::FOperationResult ULevelEditorCameraEditorState::RestoreState() const
 {
-	TSharedPtr<SLevelViewport> LevelViewport = InternalEditorLevelLibrary::GetActiveLevelViewport();
-	if (!LevelViewport.IsValid())
+	FLevelEditorModule& LevelEditorModule = FModuleManager::LoadModuleChecked<FLevelEditorModule>("LevelEditor");
+	if (TSharedPtr<ILevelEditor> FirstLevelEditor = LevelEditorModule.GetFirstLevelEditor())
+	{
+		RestoreCameraState(FirstLevelEditor);
+	}
+	else if (!OnLevelEditorCreatedDelegateHandle.IsValid())
 	{
-		return FOperationResult(FOperationResult::Failure, LOCTEXT("RestoreStateFailure_NoActiveViewport", "No active viewport"));
+		// Camera state will be applied as soon as the level viewport is created
+		OnLevelEditorCreatedDelegateHandle = LevelEditorModule.OnLevelEditorCreated().AddUObject(this, &ULevelEditorCameraEditorState::RestoreCameraState);
 	}
 
-	FLevelEditorViewportClient& LevelViewportClient = LevelViewport->GetLevelViewportClient();
+	return FOperationResult(FOperationResult::Success, FText::Format(LOCTEXT("RestoreStateSuccess", "Location=({0}) Rotation=({1}) FOV={2}"), FText::FromString(CameraLocation.ToString()), FText::FromString(CameraRotation.ToString()), CameraFOVAngle));
+}
+
+void ULevelEditorCameraEditorState::RestoreCameraState(TSharedPtr<ILevelEditor> LevelEditor) const
+{
+	// We want the delegate to be called once, so unregister ourself on the first call.
+	if (OnLevelEditorCreatedDelegateHandle.IsValid())
+	{
+		FLevelEditorModule& LevelEditorModule = FModuleManager::LoadModuleChecked<FLevelEditorModule>("LevelEditor");
+		LevelEditorModule.OnLevelEditorCreated().Remove(OnLevelEditorCreatedDelegateHandle);
+		OnLevelEditorCreatedDelegateHandle.Reset();
+	}
+
+	TSharedPtr<SLevelViewport> ActiveLevelViewport = LevelEditor->GetActiveViewportInterface();
+	FLevelEditorViewportClient& LevelViewportClient = ActiveLevelViewport->GetLevelViewportClient();
 	LevelViewportClient.SetViewLocation(CameraLocation);
 	if (!LevelViewportClient.IsOrtho())
 	{
@@ -70,8 +89,6 @@ UEditorState::FOperationResult ULevelEditorCameraEditorState::RestoreState() con
 	LevelViewportClient.Invalidate();
 
 	FEditorDelegates::OnEditorCameraMoved.Broadcast(CameraLocation, CameraRotation, LevelViewportClient.ViewportType, LevelViewportClient.ViewIndex);
-
-	return FOperationResult(FOperationResult::Success, FText::Format(LOCTEXT("RestoreStateSuccess", "Location=({0}) Rotation=({1}) FOV={2}"), FText::FromString(CameraLocation.ToString()), FText::FromString(CameraRotation.ToString()), CameraFOVAngle));
 }
 
 #undef LOCTEXT_NAMESPACE
diff --git a/Engine/Source/Editor/LevelEditor/Public/LevelEditorCameraEditorState.h b/Engine/Source/Editor/LevelEditor/Public/LevelEditorCameraEditorState.h
--- a/Engine/Source/Editor/LevelEditor/Public/LevelEditorCameraEditorState.h
+++ b/Engine/Source/Editor/LevelEditor/Public/LevelEditorCameraEditorState.h
@@ -5,6 +5,7 @@
 #include "EditorState/WorldEditorState.h"
 #include "LevelEditorCameraEditorState.generated.h"
 
+class ILevelEditor;
 class SLevelViewport;
 
 UCLASS(MinimalAPI)
@@ -25,6 +26,8 @@ private:
 	virtual FOperationResult RestoreState() const override;
 	//~ End UEditorState Interface
 
+	void RestoreCameraState(TSharedPtr<ILevelEditor> LevelEditor) const;
+
 private:
 	// Camera position
 	UPROPERTY(EditAnywhere, Category = Camera)
@@ -37,4 +40,6 @@ private:
 	// Camera FOV Angle
 	UPROPERTY(EditAnywhere, Category = Camera)
 	float CameraFOVAngle;
+
+	mutable FDelegateHandle OnLevelEditorCreatedDelegateHandle;
 };
diff --git a/Engine/Source/Editor/UnrealEd/Private/Editor.cpp b/Engine/Source/Editor/UnrealEd/Private/Editor.cpp
--- a/Engine/Source/Editor/UnrealEd/Private/Editor.cpp
+++ b/Engine/Source/Editor/UnrealEd/Private/Editor.cpp
@@ -178,6 +178,7 @@ FEditorDelegates::FOnViewAssetIdentifiers				FEditorDelegates::OnOpenAssetAudit;
 FEditorDelegates::FOnViewAssetIdentifiers				FEditorDelegates::OnEditAssetIdentifiers;
 FEditorDelegates::FOnRestartRequested					FEditorDelegates::OnRestartRequested;
 FEditorDelegates::FOnEditorBoot							FEditorDelegates::OnEditorBoot;
+FEditorDelegates::FOnEditorLoadDefaultStartupMap		FEditorDelegates::OnEditorLoadDefaultStartupMap;
 FEditorDelegates::FOnEditorInitialized					FEditorDelegates::OnEditorInitialized;
 FSimpleMulticastDelegate								FEditorDelegates::OnEditorPreExit;
 FEditorDelegates::FOnExternalContentResolved			FEditorDelegates::OnExternalContentResolved;
diff --git a/Engine/Source/Editor/UnrealEd/Private/UnrealEdMisc.cpp b/Engine/Source/Editor/UnrealEd/Private/UnrealEdMisc.cpp
--- a/Engine/Source/Editor/UnrealEd/Private/UnrealEdMisc.cpp
+++ b/Engine/Source/Editor/UnrealEd/Private/UnrealEdMisc.cpp
@@ -367,10 +367,16 @@ void FUnrealEdMisc::OnInit()
 		{
 			if (!bMapLoaded && GEditor)
 			{
-				if (GetDefault<UEditorLoadingSavingSettings>()->LoadLevelAtStartup != ELoadLevelAtStartup::None)
+				FCanLoadMap CanLoadDefaultStartupMap;
+				FEditorDelegates::OnEditorLoadDefaultStartupMap.Broadcast(CanLoadDefaultStartupMap);
+
+				if (CanLoadDefaultStartupMap.Get())
 				{
-					FEditorFileUtils::LoadDefaultMapAtStartup();
-					BeginPerformanceSurvey();
+					if (GetDefault<UEditorLoadingSavingSettings>()->LoadLevelAtStartup != ELoadLevelAtStartup::None)
+					{
+						FEditorFileUtils::LoadDefaultMapAtStartup();
+						BeginPerformanceSurvey();
+					}
 				}
 			}
 		}
diff --git a/Engine/Source/Editor/UnrealEd/Public/Editor.h b/Engine/Source/Editor/UnrealEd/Public/Editor.h
--- a/Engine/Source/Editor/UnrealEd/Public/Editor.h
+++ b/Engine/Source/Editor/UnrealEd/Public/Editor.h
@@ -180,6 +180,8 @@ struct FEditorDelegates
 	DECLARE_DELEGATE_RetVal_OneParam(bool /*bSuccess*/, FOnRestartRequested, const FString& /*ProjectName*/);
 	/** delegate for when the editor has booted */
 	DECLARE_MULTICAST_DELEGATE_OneParam(FOnEditorBoot, double Duration);
+	/** delegate for when the editor is about to load the default startup map - allows delegates to cancel the operation */
+	DECLARE_MULTICAST_DELEGATE_OneParam(FOnEditorLoadDefaultStartupMap, FCanLoadMap&);
 	/** delegate for when the editor has fully initialized */
 	DECLARE_MULTICAST_DELEGATE_OneParam(FOnEditorInitialized, double Duration);
 	/** delegate when external content resolves and can replace placeholder data */
@@ -378,6 +380,8 @@ struct FEditorDelegates
 	static UNREALED_API FOnRestartRequested OnRestartRequested;
 	/** Called when the editor has booted */
 	static UNREALED_API FOnEditorBoot OnEditorBoot;
+	/** Called when the editor is about to load the default startup map */
+	static UNREALED_API FOnEditorLoadDefaultStartupMap OnEditorLoadDefaultStartupMap;
 	/** Called when the editor has initialized */
 	static UNREALED_API FOnEditorInitialized OnEditorInitialized;
 	/** Called before the editor exits. This is called before any editor subsystems are deinitialized, and also before FCoreDelegates::OnEnginePreExit */
diff --git a/Engine/Source/Editor/WorldBookmark/Private/WorldBookmark/WorldBookmark.cpp b/Engine/Source/Editor/WorldBookmark/Private/WorldBookmark/WorldBookmark.cpp
--- a/Engine/Source/Editor/WorldBookmark/Private/WorldBookmark/WorldBookmark.cpp
+++ b/Engine/Source/Editor/WorldBookmark/Private/WorldBookmark/WorldBookmark.cpp
@@ -17,7 +17,7 @@
 
 #include UE_INLINE_GENERATED_CPP_BY_NAME(WorldBookmark)
 
-DEFINE_LOG_CATEGORY_STATIC(LogWorldBookmark, All, All);
+DEFINE_LOG_CATEGORY(LogWorldBookmark);
 
 #define LOCTEXT_NAMESPACE "WorldBookmark"
 
diff --git a/Engine/Source/Editor/WorldBookmark/Private/WorldBookmark/WorldBookmarkEditorSettings.cpp b/Engine/Source/Editor/WorldBookmark/Private/WorldBookmark/WorldBookmarkEditorSettings.cpp
--- a/Engine/Source/Editor/WorldBookmark/Private/WorldBookmark/WorldBookmarkEditorSettings.cpp
+++ b/Engine/Source/Editor/WorldBookmark/Private/WorldBookmark/WorldBookmarkEditorSettings.cpp
@@ -74,4 +74,5 @@ const TArray<FWorldBookmarkCategory>& UWorldBookmarkEditorSettings::GetCategorie
 UWorldBookmarkEditorPerProjectUserSettings::UWorldBookmarkEditorPerProjectUserSettings()
 {
 	bEnableDefaultBookmarks = true;
+	bEnableHomeBookmark = true;
 }
diff --git a/Engine/Source/Editor/WorldBookmark/Private/WorldBookmark/WorldBookmarkEditorSettings.h b/Engine/Source/Editor/WorldBookmark/Private/WorldBookmark/WorldBookmarkEditorSettings.h
--- a/Engine/Source/Editor/WorldBookmark/Private/WorldBookmark/WorldBookmarkEditorSettings.h
+++ b/Engine/Source/Editor/WorldBookmark/Private/WorldBookmark/WorldBookmarkEditorSettings.h
@@ -44,4 +44,12 @@ public:
 	/** When enabled, the default bookmark will be applied when loading a level */
 	UPROPERTY(config, EditAnywhere, Category="Default Level Bookmark")
 	bool bEnableDefaultBookmarks;
+
+	/** When enabled, the home bookmark will be applied when starting the editor without specifying a map */
+	UPROPERTY(config, EditAnywhere, Category = "Home Bookmark")
+	bool bEnableHomeBookmark;
+
+	/** Bookmark to be applied when starting the editor without specifying a map */
+	UPROPERTY(config, EditAnywhere, Category = "Home Bookmark", meta = (EditCondition = "bEnableHomeBookmark"))
+	TSoftObjectPtr<UWorldBookmark> HomeBookmark;
 };
diff --git a/Engine/Source/Editor/WorldBookmark/Private/WorldBookmark/WorldBookmarkModule.cpp b/Engine/Source/Editor/WorldBookmark/Private/WorldBookmark/WorldBookmarkModule.cpp
--- a/Engine/Source/Editor/WorldBookmark/Private/WorldBookmark/WorldBookmarkModule.cpp
+++ b/Engine/Source/Editor/WorldBookmark/Private/WorldBookmark/WorldBookmarkModule.cpp
@@ -47,8 +47,12 @@ void FWorldBookmarkModule::StartupModule()
 		StructsToUnregisterOnShutdown.Add(FWorldBookmarkCategory::StaticStruct()->GetFName());
 		PropertyModule.RegisterCustomPropertyTypeLayout(FWorldBookmarkCategory::StaticStruct()->GetFName(), FOnGetPropertyTypeCustomizationInstance::CreateStatic(&FWorldBookmarkCategoryCustomization::MakeInstance));
 
+		// When deleting a world, also delete associated bookmarks
 		OnAddExtraObjectsToDeleteDelegateHandle = FEditorDelegates::OnAddExtraObjectsToDelete.AddRaw(this, &FWorldBookmarkModule::OnAddExtraObjectsToDelete);
 
+		// Override the loading of the default startup map if the user specified a Home Bookmark
+		OnEditorLoadDefaultStartupMapHandle = FEditorDelegates::OnEditorLoadDefaultStartupMap.AddRaw(this, &FWorldBookmarkModule::OnEditorLoadDefaultStartupMap);
+
 		// Listen for map change events
 		OnMapChangedHandle = LevelEditorModule.OnMapChanged().AddRaw(this, &FWorldBookmarkModule::OnMapChanged);
 
@@ -92,6 +96,7 @@ void FWorldBookmarkModule::ShutdownModule()
 		StructsToUnregisterOnShutdown.Empty();
 
 		FEditorDelegates::OnAddExtraObjectsToDelete.Remove(OnAddExtraObjectsToDeleteDelegateHandle);
+		FEditorDelegates::OnEditorLoadDefaultStartupMap.Remove(OnEditorLoadDefaultStartupMapHandle);
 	}
 }
 
@@ -209,6 +214,7 @@ void FWorldBookmarkModule::OnMapChanged(UWorld* InWorld, EMapChangeType InMapCha
 				// At this point, make sure we are really dealing with a World Bookmark. Otherwise, do nothing.
 				if (UWorldBookmark* WorldBookmark = Cast<UWorldBookmark>(InWorld->GetWorldSettings()->GetDefaultBookmark()))
 				{
+					UE_LOG(LogWorldBookmark, Display, TEXT("Loading default bookmark %s"), *WorldBookmark->GetPathName());
 					WorldBookmark->Load();
 				}
 			}
@@ -261,4 +267,25 @@ void FWorldBookmarkModule::ShowInvalidDefaultBookmarkNotification(const FText& I
 	FSlateNotificationManager::Get().AddNotification(Info);
 }
 
+void FWorldBookmarkModule::OnEditorLoadDefaultStartupMap(FCanLoadMap& InOutCanLoadDefaultStartupMap)
+{
+	const UWorldBookmarkEditorPerProjectUserSettings* Settings = GetDefault<UWorldBookmarkEditorPerProjectUserSettings>();
+	if (Settings->bEnableHomeBookmark)
+	{
+		if (UWorldBookmark* HomeBookmark = Settings->HomeBookmark.LoadSynchronous())
+		{
+			UE_LOG(LogWorldBookmark, Display, TEXT("Loading home bookmark %s"), *Settings->HomeBookmark.ToString());
+			HomeBookmark->Load();
+
+			// If the bookmark world was loaded successfully, prevent the loading of the default startup map.
+			// Verify that the current map now matches the one from the bookmark
+			TSoftObjectPtr<UWorld> BookmarkWorld = HomeBookmark->GetEditorState<UWorldEditorState>()->GetStateWorld();
+			if (GEditor->GetEditorWorldContext().World() == BookmarkWorld)
+			{
+				InOutCanLoadDefaultStartupMap.SetFalse();
+			}
+		}
+	}
+}
+
 #undef LOCTEXT_NAMESPACE
diff --git a/Engine/Source/Editor/WorldBookmark/Private/WorldBookmark/WorldBookmarkModule.h b/Engine/Source/Editor/WorldBookmark/Private/WorldBookmark/WorldBookmarkModule.h
--- a/Engine/Source/Editor/WorldBookmark/Private/WorldBookmark/WorldBookmarkModule.h
+++ b/Engine/Source/Editor/WorldBookmark/Private/WorldBookmark/WorldBookmarkModule.h
@@ -36,6 +36,7 @@ private:
 	void OnAddExtraObjectsToDelete(const TArray<UObject*>& InObjectsToDelete, TSet<UObject*>& OutSecondaryObjects);
 	void OnMapChanged(UWorld* InWorld, EMapChangeType InMapChangeType);
 	void OnDefaultBookmarkChanged(AWorldSettings* InWorldSettings);
+	void OnEditorLoadDefaultStartupMap(FCanLoadMap& InOutCanLoadDefaultStartupMap);
 
 	bool IsDefaultBookmarkValid(const AWorldSettings* InWorldSettings) const;
 	void ShowInvalidDefaultBookmarkNotification(const FText& InNotificationTitle) const;
@@ -43,6 +44,7 @@ private:
 	FDelegateHandle OnAddExtraObjectsToDeleteDelegateHandle;
 	FDelegateHandle OnMapChangedHandle;
 	FDelegateHandle OnDefaultBookmarkChangedHandle;
+	FDelegateHandle OnEditorLoadDefaultStartupMapHandle;
 
 	TArray<FName> ClassesToUnregisterOnShutdown;
 	TArray<FName> StructsToUnregisterOnShutdown;
diff --git a/Engine/Source/Editor/WorldBookmark/Public/WorldBookmark/WorldBookmark.h b/Engine/Source/Editor/WorldBookmark/Public/WorldBookmark/WorldBookmark.h
--- a/Engine/Source/Editor/WorldBookmark/Public/WorldBookmark/WorldBookmark.h
+++ b/Engine/Source/Editor/WorldBookmark/Public/WorldBookmark/WorldBookmark.h
@@ -8,6 +8,8 @@
 #include "Engine/BookmarkBase.h"
 #include "WorldBookmark.generated.h"
 
+DECLARE_LOG_CATEGORY_EXTERN(LogWorldBookmark, Log, All);
+
 /**
  * World Bookmark Category
  */
