World Bookmark - Added the ability to capture / restore bookmark as strings through console commands
* The bookmark data (various states of the editor) is serialized, compressed and converted to base64

Console commands are:
* WorldBookmark.Capture - Will log the bookmark to the console, as a string
* WorldBookmark.CaptureToClipboard - Capture the bookmark directly to the clipboard
* WorldBookmark.Restore - Restore the bookmark from a string
* WorldBookmark.RestoreFromClipboard - Same, but without having to paste the bookmark text

diff --git a/Engine/Source/Editor/DataLayerEditor/Private/DataLayer/DataLayerEditorState.cpp b/Engine/Source/Editor/DataLayerEditor/Private/DataLayer/DataLayerEditorState.cpp
--- a/Engine/Source/Editor/DataLayerEditor/Private/DataLayer/DataLayerEditorState.cpp
+++ b/Engine/Source/Editor/DataLayerEditor/Private/DataLayer/DataLayerEditorState.cpp
@@ -48,7 +48,8 @@ UEditorState::FOperationResult UDataLayerEditorState::CaptureState()
 		return true;
 	});
 
-	return FOperationResult(FOperationResult::Success, FText::Format(LOCTEXT("CaptureStateSuccess", "LoadedDataLayers={0}, NotLoadedDataLayers={1}"), LoadedDataLayers.Num(), NotLoadedDataLayers.Num()));
+	FOperationResult::EResult OperationResult = LoadedDataLayers.IsEmpty() && NotLoadedDataLayers.IsEmpty() ? FOperationResult::Skipped : FOperationResult::Success;
+	return FOperationResult(OperationResult, FText::Format(LOCTEXT("CaptureStateSuccess", "LoadedDataLayers={0}, NotLoadedDataLayers={1}"), LoadedDataLayers.Num(), NotLoadedDataLayers.Num()));
 }
 
 UEditorState::FOperationResult UDataLayerEditorState::RestoreState() const
diff --git a/Engine/Source/Editor/UnrealEd/Private/EditorState/EditorStateSubsystem.cpp b/Engine/Source/Editor/UnrealEd/Private/EditorState/EditorStateSubsystem.cpp
--- a/Engine/Source/Editor/UnrealEd/Private/EditorState/EditorStateSubsystem.cpp
+++ b/Engine/Source/Editor/UnrealEd/Private/EditorState/EditorStateSubsystem.cpp
@@ -87,18 +87,18 @@ void UEditorStateSubsystem::CaptureEditorState(FEditorStateCollection& OutState,
 
 	if (InEditorStateTypeFilter.IsEmpty())
 	{
-		OutState.EditorStates.Reset();
+		OutState.States.Reset();
 
 		for (TSubclassOf<UEditorState> EditorStateType : RegisteredEditorStateTypes)
 		{
-			OutState.EditorStates.Emplace(EditorStateType, NewObject<UEditorState>(InStateOuter, EditorStateType));
+			OutState.States.Emplace(NewObject<UEditorState>(InStateOuter, EditorStateType));
 		}
 	}
 	else
 	{
 		for (TSubclassOf<UEditorState> EditorStateType : InEditorStateTypeFilter)
 		{
-			OutState.EditorStates.Emplace(EditorStateType, NewObject<UEditorState>(InStateOuter, EditorStateType));
+			OutState.States.Emplace(NewObject<UEditorState>(InStateOuter, EditorStateType));
 		}
 	}
 
@@ -122,16 +122,16 @@ void UEditorStateSubsystem::CaptureEditorState(FEditorStateCollection& OutState,
 		if (!bCaptureSuccess)
 		{
 			StateToCapture->Rename(nullptr, GetTransientPackage());
-			OutState.EditorStates.FindChecked(StateToCapture->GetClass()) = nullptr;
+			OutState.GetStateChecked(StateToCapture->GetClass()) = nullptr;
 		}
 
 		return bCaptureSuccess;
 	}, InEditorStateTypeFilter);
 
 	// Remove invalidated states
-	for (FEditorStateCollection::FEditorStateMap::TIterator It(OutState.EditorStates); It; ++It)
+	for (TArray<TObjectPtr<UEditorState>>::TIterator It(OutState.States); It; ++It)
 	{
-		if (It.Value() == nullptr)
+		if (It->Get() == nullptr)
 		{
 			It.RemoveCurrent();
 		}
diff --git a/Engine/Source/Editor/UnrealEd/Public/EditorState/EditorStateCollection.h b/Engine/Source/Editor/UnrealEd/Public/EditorState/EditorStateCollection.h
--- a/Engine/Source/Editor/UnrealEd/Public/EditorState/EditorStateCollection.h
+++ b/Engine/Source/Editor/UnrealEd/Public/EditorState/EditorStateCollection.h
@@ -17,21 +17,21 @@ struct FEditorStateCollection
 	template <typename TEditorStateType>
 	typename TEnableIf<TIsDerivedFrom<TEditorStateType, UEditorState>::IsDerived, bool>::Type HasState() const
 	{
-		return EditorStates.Contains(TEditorStateType::StaticClass());
+		return GetState(TEditorStateType::StaticClass()) != nullptr;
 	}
 
 	template <typename TEditorStateType>
 	typename TEnableIf<TIsDerivedFrom<TEditorStateType, UEditorState>::IsDerived, const TEditorStateType*>::Type GetState() const
 	{
-		const TObjectPtr<UEditorState>* EditorStatePtr = EditorStates.Find(TEditorStateType::StaticClass());
-		return EditorStatePtr ? CastChecked<TEditorStateType>((*EditorStatePtr).Get()) : nullptr;
+		const UEditorState* EditorState = GetState(TEditorStateType::StaticClass());
+		return EditorState ? CastChecked<TEditorStateType>(EditorState) : nullptr;
 	}
 
 	template <typename TEditorStateType>
 	typename TEnableIf<TIsDerivedFrom<TEditorStateType, UEditorState>::IsDerived, const TEditorStateType&>::Type GetStateChecked() const
 	{
-		const TObjectPtr<UEditorState>& EditorState = EditorStates.FindChecked(TEditorStateType::StaticClass());
-		return *CastChecked<TEditorStateType>(EditorState.Get());
+		const UEditorState* EditorState = GetState(TEditorStateType::StaticClass());
+		return *CastChecked<TEditorStateType>(EditorState);
 	}
 
 	/** Iterate over each state, making sure iteration is done over dependant states first */
@@ -48,19 +48,23 @@ struct FEditorStateCollection
 
 	bool HasStates() const
 	{
-		return !EditorStates.IsEmpty();
+		return !States.IsEmpty();
 	}
 
 	TArray<UEditorState*> GetStates() const
 	{
-		TArray<UEditorState*> Values;
+		return States;
+	}
 
-		for (const TPair<TSubclassOf<UEditorState>, TObjectPtr<UEditorState>>& KeyValue : EditorStates)
+	void PostSerialize(const FArchive& Ar)
+	{
+PRAGMA_DISABLE_DEPRECATION_WARNINGS
+		for (const TPair<TSubclassOf<UEditorState>, TObjectPtr<UEditorState>>& KeyValue : EditorStates_DEPRECATED)
 		{
-			Values.Add(KeyValue.Value);
+			States.Emplace(KeyValue.Value);
 		}
-	
-		return Values;
+		EditorStates_DEPRECATED.Reset();
+PRAGMA_ENABLE_DEPRECATION_WARNINGS
 	}
 
 private:
@@ -83,7 +87,7 @@ private:
 					// Process the dependencies - Stop if any error is encountered
 					for (TSubclassOf<UEditorState> DependencyType : StateToProcess->GetDependencies())
 					{
-						UEditorState* DependencyState = Self->EditorStates.FindRef(DependencyType);
+						UEditorState* DependencyState = Self->GetState(DependencyType);
 						bProcessedDependenciesSuccessfully = DependencyState && ProcessStateRef(DependencyState, ProcessStateRef);
 
 						if (!bProcessedDependenciesSuccessfully)
@@ -106,21 +110,69 @@ private:
 			return ProcessStateImpl(StateToProcess, ProcessStateImpl);
 		};
 
-		for (const auto& KeyValue : Self->EditorStates)
+		for (const TObjectPtr<UEditorState>& State : Self->States)
 		{
-			if (!InEditorStateTypeFilter.IsEmpty() && !InEditorStateTypeFilter.Contains(KeyValue.Value->GetClass()))
+			if (!InEditorStateTypeFilter.IsEmpty() && !InEditorStateTypeFilter.Contains(State->GetClass()))
 			{
 				continue;
 			}
 
-			ProcessState(KeyValue.Value);
+			ProcessState(State.Get());
 		}
 	}
 
-	typedef TMap<TSubclassOf<UEditorState>, TObjectPtr<UEditorState>> FEditorStateMap;
+	UEditorState* GetState(TSubclassOf<UEditorState> InStateType) const
+	{
+		const TObjectPtr<UEditorState>* Result = States.FindByPredicate([&InStateType](const TObjectPtr<UEditorState> InEditorState)
+		{
+			return InEditorState && InEditorState->IsA(InStateType);
+		});
+		return Result ? *Result : nullptr;
+	}
+
+	TObjectPtr<UEditorState>& GetStateChecked(TSubclassOf<UEditorState> InStateType)
+	{
+		TObjectPtr<UEditorState>* Result = States.FindByPredicate([&InStateType](TObjectPtr<UEditorState> InEditorState)
+		{
+			return InEditorState && InEditorState->IsA(InStateType);
+		});
+		check(Result);
+		return *Result;
+	}
 
+	UE_DEPRECATED(5.6, "Use the States array")
 	UPROPERTY()
-	TMap<TSubclassOf<UEditorState>, TObjectPtr<UEditorState>> EditorStates;
+	TMap<TSubclassOf<UEditorState>, TObjectPtr<UEditorState>> EditorStates_DEPRECATED;
+
+	UPROPERTY(Instanced)
+	TArray<TObjectPtr<UEditorState>> States;
 	
 	friend class UEditorStateSubsystem;
 };
+
+// Add traits to let the engine know we have a PostSerialize() method on our struct
+template<>
+struct TStructOpsTypeTraits<FEditorStateCollection> : public TStructOpsTypeTraitsBase2<FEditorStateCollection>
+{
+	enum
+	{
+		WithPostSerialize = true
+	};
+};
+
+// Wrapper to avoid GC of a standalone FEditorStateCollection (as it contains UObjects)
+class FEditorStateCollectionGCObject : public FGCObject
+{
+public:
+	FEditorStateCollection EditorStateCollection;
+
+	virtual void AddReferencedObjects(FReferenceCollector& Collector) override
+	{
+		Collector.AddPropertyReferences(FEditorStateCollection::StaticStruct(), &EditorStateCollection);
+	}
+
+	virtual FString GetReferencerName() const override
+	{
+		return TEXT("FEditorStateCollection");
+	}
+};
diff --git a/Engine/Source/Editor/WorldBookmark/Private/WorldBookmark/Browser/SWorldBookmarkBrowser.cpp b/Engine/Source/Editor/WorldBookmark/Private/WorldBookmark/Browser/SWorldBookmarkBrowser.cpp
--- a/Engine/Source/Editor/WorldBookmark/Private/WorldBookmark/Browser/SWorldBookmarkBrowser.cpp
+++ b/Engine/Source/Editor/WorldBookmark/Private/WorldBookmark/Browser/SWorldBookmarkBrowser.cpp
@@ -1264,19 +1264,19 @@ void SWorldBookmarkBrowser::DeleteSelectedItem()
 
 static bool AskForWorldChangeConfirmation(UWorldBookmark* WorldBookmark)
 {
-	TSoftObjectPtr<UWorld> BookmarkWorld = WorldBookmark->GetEditorState<UWorldEditorState>()->GetStateWorld();
-	if (BookmarkWorld != GetCurrentWorld())
+	if (const UWorldEditorState* WorldEditorState = WorldBookmark->GetEditorState<UWorldEditorState>())
 	{
-		const FText MessageBoxTitle = LOCTEXT("WorldBookmark_DifferentWorldDlg_Title", "Open Bookmark World?");
-		const FText MessageBoxText = LOCTEXT("WorldBookmark_DifferentWorldDlg_Text", "This bookmark is for a different world, are you sure you want to open it?");
-		const EAppReturnType::Type Response = FMessageDialog::Open(EAppMsgType::YesNo, MessageBoxText, MessageBoxTitle);
-		if (Response != EAppReturnType::Yes)
+		TSoftObjectPtr<UWorld> BookmarkWorld = WorldEditorState->GetStateWorld();
+		if (BookmarkWorld != GetCurrentWorld())
 		{
-			return false;
+			const FText MessageBoxTitle = LOCTEXT("WorldBookmark_DifferentWorldDlg_Title", "Open Bookmark World?");
+			const FText MessageBoxText = LOCTEXT("WorldBookmark_DifferentWorldDlg_Text", "This bookmark is for a different world, are you sure you want to open it?");
+			const EAppReturnType::Type Response = FMessageDialog::Open(EAppMsgType::YesNo, MessageBoxText, MessageBoxTitle);
+			return Response == EAppReturnType::Yes;
 		}
 	}
 
-	return true;
+	return false;
 }
 
 bool SWorldBookmarkBrowser::CanGoToSelectedBookmarkLocation() const
diff --git a/Engine/Source/Editor/WorldBookmark/Private/WorldBookmark/WorldBookmarkModule.cpp b/Engine/Source/Editor/WorldBookmark/Private/WorldBookmark/WorldBookmarkModule.cpp
--- a/Engine/Source/Editor/WorldBookmark/Private/WorldBookmark/WorldBookmarkModule.cpp
+++ b/Engine/Source/Editor/WorldBookmark/Private/WorldBookmark/WorldBookmarkModule.cpp
@@ -14,7 +14,11 @@
 #include "EditorState/WorldEditorState.h"
 #include "Framework/Notifications/NotificationManager.h"
 #include "GameFramework/WorldSettings.h"
+#include "HAL/PlatformApplicationMisc.h"
+#include "JsonObjectConverter.h"
 #include "LevelEditor.h"
+#include "Misc/Base64.h"
+#include "Misc/Compression.h"
 #include "Misc/MessageDialog.h"
 #include "Modules/ModuleManager.h"
 #include "ObjectTools.h"
@@ -30,6 +34,65 @@ IMPLEMENT_MODULE(FWorldBookmarkModule, WorldBookmark);
 
 const FName WorldBookmarkBrowserTabId("WorldBookmarkBrowser");
 
+static FAutoConsoleCommand WorldBookmarkCaptureToStringCommand(
+	TEXT("WorldBookmark.Capture"),
+	TEXT("Capture the current state of the editor and log it to the console."),
+	FConsoleCommandWithArgsDelegate::CreateLambda([](const TArray<FString>& Args)
+	{
+		FString Str = FWorldBookmarkModule::CaptureToString();
+		UE_LOG(LogWorldBookmark, Display, TEXT("Restore the bookmark with 'WorldBookmark.Restore' and the following argument:\n%s"), *Str);
+	})
+);
+
+static FAutoConsoleCommand WorldBookmarkCaptureToClipboardCommand(
+	TEXT("WorldBookmark.CaptureToClipboard"),
+	TEXT("Capture the current state of the editor and copy it to the clipboard."),
+	FConsoleCommandWithArgsDelegate::CreateLambda([](const TArray<FString>& Args)
+	{
+		FWorldBookmarkModule::CaptureToClipboard();
+		UE_LOG(LogWorldBookmark, Display, TEXT("WorldBookmark captured to clipboard"));
+	})
+);
+
+static FAutoConsoleCommand WorldBookmarkRestoreFromStringCommand(
+	TEXT("WorldBookmark.Restore"),
+	TEXT("Restore a bookmark from the text previously obtained from WorldBookmark.Capture."),
+	FConsoleCommandWithArgsDelegate::CreateLambda([](const TArray<FString>& Args)
+	{
+		TArray<FString> History;
+		IConsoleManager::Get().GetConsoleHistory(TEXT(""), History);
+		
+		bool bRestoredFromConsoleHistory = false;
+
+		// Try to retrieve a multiline command from the console history, as bookmarks strings may be pasted with newlines in them
+		if (!History.IsEmpty())
+		{
+			FString BookmarkString = History.Last();
+			if (BookmarkString.RemoveFromStart(TEXT("WorldBookmark.Restore ")))
+			{
+				bRestoredFromConsoleHistory = true;
+				FWorldBookmarkModule::RestoreFromString(BookmarkString);
+			}
+		}
+		
+		// Failed to obtain the command from the console history
+		// Fallback to the arguments
+		if (!bRestoredFromConsoleHistory && Args.Num() == 1)
+		{
+			FWorldBookmarkModule::RestoreFromString(Args[0]);
+		}
+	})
+);
+
+static FAutoConsoleCommand WorldBookmarkRestoreFromClipboardCommand(
+	TEXT("WorldBookmark.RestoreFromClipboard"),
+	TEXT("Restore a bookmark from the clipboard."),
+	FConsoleCommandWithArgsDelegate::CreateLambda([](const TArray<FString>& Args)
+	{
+		FWorldBookmarkModule::RestoreFromClipboard();
+	})
+);
+
 void FWorldBookmarkModule::StartupModule()
 {
 	if (GIsEditor && !IsRunningCommandlet())
@@ -244,7 +307,8 @@ bool FWorldBookmarkModule::IsDefaultBookmarkValid(const AWorldSettings* InWorldS
 	}
 
 	// Validate that the world bookmark's world is our actual world
-	TSoftObjectPtr<UWorld> BookmarkWorld = WorldBookmark->GetEditorState<UWorldEditorState>()->GetStateWorld();
+	const UWorldEditorState* WorldEditorState = WorldBookmark->GetEditorState<UWorldEditorState>();
+	TSoftObjectPtr<UWorld> BookmarkWorld = WorldEditorState ? WorldEditorState->GetStateWorld() : nullptr;
 	return BookmarkWorld == InWorldSettings->GetWorld();
 }
 
@@ -279,13 +343,290 @@ void FWorldBookmarkModule::OnEditorLoadDefaultStartupMap(FCanLoadMap& InOutCanLo
 
 			// If the bookmark world was loaded successfully, prevent the loading of the default startup map.
 			// Verify that the current map now matches the one from the bookmark
-			TSoftObjectPtr<UWorld> BookmarkWorld = HomeBookmark->GetEditorState<UWorldEditorState>()->GetStateWorld();
-			if (GEditor->GetEditorWorldContext().World() == BookmarkWorld)
+			if (const UWorldEditorState* WorldEditorState = HomeBookmark->GetEditorState<UWorldEditorState>())
 			{
-				InOutCanLoadDefaultStartupMap.SetFalse();
+				TSoftObjectPtr<UWorld> BookmarkWorld = WorldEditorState->GetStateWorld();
+				if (GEditor->GetEditorWorldContext().World() == BookmarkWorld)
+				{
+					InOutCanLoadDefaultStartupMap.SetFalse();
+				}
+			}
+		}
+	}
+}
+
+struct FBookmarkTextExport
+{
+	static TSharedPtr<FJsonValue> ExportPropertyToJSon(FProperty* Property, const void* Value)
+	{
+		static const UScriptStruct* BoxStruct = FindObject<UScriptStruct>(nullptr, TEXT("/Script/CoreUObject.Box"), true);
+		static const UScriptStruct* VectorStruct = FindObject<UScriptStruct>(nullptr, TEXT("/Script/CoreUObject.Vector"), true);
+		static const UScriptStruct* RotatorStruct = FindObject<UScriptStruct>(nullptr, TEXT("/Script/CoreUObject.Rotator"), true);
+
+		if (FStructProperty* PropertyAsStruct = CastField<FStructProperty>(Property))
+		{
+			auto Round = [](double Val)
+			{
+				return FMath::RoundToInt(Val * 100) / 100.0;
+			};
+
+			if (PropertyAsStruct->Struct == BoxStruct)
+			{
+				FBox Box;
+				PropertyAsStruct->CopySingleValue(&Box, Value);
+				FString AsString = Box.IsValid ? FString::Printf(TEXT("%.2f %.2f %.2f %.2f %.2f %.2f"), Round(Box.Min.X), Round(Box.Min.Y), Round(Box.Min.Z), Round(Box.Max.X), Round(Box.Max.Y), Round(Box.Min.Z)) : TEXT("");
+				return MakeShared<FJsonValueString>(*AsString);
+			}
+			else if (PropertyAsStruct->Struct == VectorStruct)
+			{
+				FVector Vec;
+				PropertyAsStruct->CopySingleValue(&Vec, Value);
+				FString AsString = !Vec.IsNearlyZero() ? FString::Printf(TEXT("%.2f %.2f %.2f"), Round(Vec.X), Round(Vec.Y), Round(Vec.Z)) : TEXT("");
+				return MakeShared<FJsonValueString>(*AsString);
+			}
+			else if (PropertyAsStruct->Struct == RotatorStruct)
+			{
+				FRotator Rot;
+				PropertyAsStruct->CopySingleValue(&Rot, Value);
+				FString AsString = !Rot.IsNearlyZero() ? FString::Printf(TEXT("%.2f %.2f %.2f"), Round(Rot.Pitch), Round(Rot.Yaw), Round(Rot.Roll)) : TEXT("");
+				return MakeShared<FJsonValueString>(*AsString);
+			}
+		}
+
+		return nullptr;
+	}
+
+	static bool ImportPropertyFromJSon(const TSharedPtr<FJsonValue>& JsonValue, FProperty* Property, void* Value)
+	{
+		static const UScriptStruct* BoxStruct = FindObject<UScriptStruct>(nullptr, TEXT("/Script/CoreUObject.Box"), true);
+		static const UScriptStruct* VectorStruct = FindObject<UScriptStruct>(nullptr, TEXT("/Script/CoreUObject.Vector"), true);
+		static const UScriptStruct* RotatorStruct = FindObject<UScriptStruct>(nullptr, TEXT("/Script/CoreUObject.Rotator"), true);
+
+		if (FStructProperty* PropertyAsStruct = CastField<FStructProperty>(Property))
+		{
+			TArray<double> DoubleArray;
+
+			auto ToDoubleArray = [](const FString& InString, TArray<double>& OutDoubles) -> int
+			{
+				TArray<FString> StringArray;
+				InString.ParseIntoArray(StringArray, TEXT(" "));
+				Algo::Transform(StringArray, OutDoubles, [](const FString& String)
+				{
+					return FCString::Atod(*String);
+				});
+				return OutDoubles.Num();
+			};
+
+			if (PropertyAsStruct->Struct == BoxStruct)
+			{
+				FBox& Box = *static_cast<FBox*>(Value);
+				if (ToDoubleArray(JsonValue->AsString(), DoubleArray) == 6)
+				{
+					FMemory::Memcpy(Value, DoubleArray.GetData(), 6 * sizeof(double));
+					Box.IsValid = true;
+				}
+				else
+				{
+					Box.IsValid = false;
+				}
+				return true;
+			}
+			else if ((PropertyAsStruct->Struct == VectorStruct) || (PropertyAsStruct->Struct == RotatorStruct))
+			{
+				if (ToDoubleArray(JsonValue->AsString(), DoubleArray) == 3)
+				{					
+					FMemory::Memcpy(Value, DoubleArray.GetData(), 3 * sizeof(double));
+				}
+				else
+				{
+					FMemory::Memzero(Value, 3 * sizeof(double));
+				}
+				return true;
 			}
 		}
+
+		return false;
+	}
+
+	int8  Version = 0;						// For possible future expansions
+	int32 UncompressedSize = 0;
+	TArray<uint8> CompressedData;
+};
+
+
+static FString FormatBookmarkString(const FString& InBookmarkString, const int32 InMaxLineLength)
+{
+	FString Result;
+	int32 InputLength = InBookmarkString.Len();
+	int32 CurrentIndex = 0;
+
+	// Process the input string in chunks of 100 characters
+	while (CurrentIndex < InputLength)
+	{
+		// Determine the length of the next chunk
+		int32 RemainingCharacters = InputLength - CurrentIndex;
+		int32 ChunkLength = FMath::Min(RemainingCharacters, InMaxLineLength);
+
+		// Extract the chunk
+		FString Chunk = InBookmarkString.Mid(CurrentIndex, ChunkLength);
+		Result += Chunk;
+
+		// Add a newline if we're not at the end of the string
+		CurrentIndex += ChunkLength;
+		if (CurrentIndex < InputLength)
+		{
+			Result += TEXT("\n");
+		}
+	}
+
+	// Check if the last line needs padding
+	int32 LastLineStart = Result.Find(TEXT("\n"), ESearchCase::IgnoreCase, ESearchDir::FromEnd) + 1;
+	if (LastLineStart == 0)
+	{
+		LastLineStart = 0; // No newlines, start at the beginning
+	}
+
+	int32 LastLineLength = Result.Len() - LastLineStart;
+	if (LastLineLength < InMaxLineLength)
+	{
+		Result += FString::ChrN(InMaxLineLength - LastLineLength, TEXT('='));
 	}
+
+	return Result;
+}
+
+FString FWorldBookmarkModule::CaptureToString()
+{
+	// Capture the editor state
+	FEditorStateCollection EditorStateCollection;
+	UEditorStateSubsystem::Get()->CaptureEditorState(EditorStateCollection, GetTransientPackage());
+	
+	// Convert to a Json string
+	FString BookmarkAsJson;
+	FJsonObjectConverter::CustomExportCallback CustomExportCallback;
+	CustomExportCallback.BindStatic(FBookmarkTextExport::ExportPropertyToJSon);
+	FJsonObjectConverter::UStructToJsonObjectString(EditorStateCollection, BookmarkAsJson, 0, CPF_Deprecated | CPF_Transient | CPF_DuplicateTransient | CPF_TextExportTransient, 0, &CustomExportCallback, false);
+	
+	// Convert string to UTF8
+	const FTCHARToUTF8 BufferString(*BookmarkAsJson);
+
+	FBookmarkTextExport BookmarkTextExport;
+	BookmarkTextExport.UncompressedSize = BufferString.Length();
+	BookmarkTextExport.CompressedData.SetNumUninitialized(BookmarkTextExport.UncompressedSize);
+
+	// Compress the UTF8 Json string using Zlib
+	int32 CompressedSize = BufferString.Length();
+	FCompression::CompressMemory(NAME_Zlib, BookmarkTextExport.CompressedData.GetData(), CompressedSize, BufferString.Get(), BufferString.Length(), COMPRESS_BiasSize);
+	BookmarkTextExport.CompressedData.SetNum(CompressedSize);
+
+	// Write all the info to a buffer
+	TArray<uint8> ExportBuffer;
+	FMemoryWriter MemoryWriter(ExportBuffer);
+	MemoryWriter << BookmarkTextExport.Version;
+	MemoryWriter << BookmarkTextExport.UncompressedSize;
+	MemoryWriter << BookmarkTextExport.CompressedData;
+
+	// Encode buffer to a base64 string
+	FString Result = FString("BM") + FBase64::Encode(ExportBuffer);
+
+	const int32 MAX_LINE_LENGTH = 100;
+	return FormatBookmarkString(Result, MAX_LINE_LENGTH);
+}
+
+void FWorldBookmarkModule::CaptureToClipboard()
+{
+	FString BookmarkAsString = CaptureToString();
+	FPlatformApplicationMisc::ClipboardCopy(*BookmarkAsString);
+}
+
+bool FWorldBookmarkModule::RestoreFromString(const FString& InBookmarkAsString)
+{
+	auto ShowInvalidBookmarkDataNotification = []()
+	{
+		// Show toast.
+		FNotificationInfo Info(LOCTEXT("BookmarkFromString_Invalid_Text", "Invalid bookmark data"));
+		Info.SubText = LOCTEXT("BookmarkFromString_Invalid_SubText", "The bookmark was not restored.");
+		Info.ExpireDuration = 3.0f;
+		Info.bFireAndForget = true;
+		Info.Image = FCoreStyle::Get().GetBrush(TEXT("MessageLog.Error"));
+		FSlateNotificationManager::Get().AddNotification(Info);
+	};
+
+	if (!InBookmarkAsString.StartsWith("BM"))
+	{
+		UE_LOG(LogWorldBookmark, Error, TEXT("RestoreFromString: Invalid bookmark string"));
+		ShowInvalidBookmarkDataNotification();
+		return false;
+	}
+
+	// Remove 'BM' header and any newlines
+	FString Base64String = InBookmarkAsString.RightChop(2);
+	Base64String.ReplaceInline(TEXT("\n"), TEXT(""));
+	Base64String.ReplaceInline(TEXT("\r"), TEXT(""));
+	Base64String.ReplaceInline(TEXT(" "), TEXT(""));
+	
+	// Decode base64 string to a buffer
+	TArray<uint8> ImportBuffer;
+	if (!FBase64::Decode(Base64String, ImportBuffer))
+	{
+		UE_LOG(LogWorldBookmark, Error, TEXT("RestoreFromString: Failed to decode base64 bookmark data"));
+		ShowInvalidBookmarkDataNotification();
+		return false;
+	}
+
+	FMemoryReader MemoryReader(ImportBuffer);
+
+	// Read all the info from the buffer
+	FBookmarkTextExport BookmarkTextImport;
+	if (MemoryReader.TotalSize() > sizeof(FBookmarkTextExport))
+	{
+		MemoryReader << BookmarkTextImport.Version;
+		MemoryReader << BookmarkTextImport.UncompressedSize;
+		MemoryReader << BookmarkTextImport.CompressedData;
+	}
+	else
+	{
+		UE_LOG(LogWorldBookmark, Error, TEXT("RestoreFromString: Failed to read compressed bookmark data"));
+		ShowInvalidBookmarkDataNotification();
+		return false;
+	}
+	
+	// Decompress the buffer to a UTF8 Json string
+	TArray<uint8> UncompressedData;
+	UncompressedData.SetNum(BookmarkTextImport.UncompressedSize);
+	if (!FCompression::UncompressMemory(NAME_Zlib, UncompressedData.GetData(), BookmarkTextImport.UncompressedSize, BookmarkTextImport.CompressedData.GetData(), BookmarkTextImport.CompressedData.Num(), COMPRESS_BiasSize))
+	{
+		UE_LOG(LogWorldBookmark, Error, TEXT("RestoreFromString: Failed to uncompress bookmark data"));
+		ShowInvalidBookmarkDataNotification();
+		return false;
+	}
+
+	// Convert UTF8 Json string to an FString
+	FUTF8ToTCHAR BufferToString((const ANSICHAR*)UncompressedData.GetData(), UncompressedData.Num());
+	FString BookmarkAsJson = FString::ConstructFromPtrSize(BufferToString.Get(), BufferToString.Length());;
+	
+	// Read the editor state from the Json
+	FEditorStateCollectionGCObject EditorStateCollectionGCObject;
+	FJsonObjectConverter::CustomImportCallback CustomImportCallback;
+	CustomImportCallback.BindStatic(FBookmarkTextExport::ImportPropertyFromJSon);
+	if (!FJsonObjectConverter::JsonObjectStringToUStruct(BookmarkAsJson, &EditorStateCollectionGCObject.EditorStateCollection, 0, CPF_Deprecated | CPF_Transient | CPF_DuplicateTransient | CPF_TextExportTransient, false, nullptr, &CustomImportCallback))
+	{
+		UE_LOG(LogWorldBookmark, Error, TEXT("RestoreFromString: Failed to restore bookmark data"));
+		ShowInvalidBookmarkDataNotification();
+		return false;
+	}
+	
+	// Restore the editor state
+	UEditorStateSubsystem::Get()->RestoreEditorState(EditorStateCollectionGCObject.EditorStateCollection);
+
+	return true;
+}
+
+bool FWorldBookmarkModule::RestoreFromClipboard()
+{
+	FString BookmarkAsString;
+	FPlatformApplicationMisc::ClipboardPaste(BookmarkAsString);
+	return RestoreFromString(BookmarkAsString);
 }
 
 #undef LOCTEXT_NAMESPACE
diff --git a/Engine/Source/Editor/WorldBookmark/Private/WorldBookmark/WorldBookmarkModule.h b/Engine/Source/Editor/WorldBookmark/Private/WorldBookmark/WorldBookmarkModule.h
--- a/Engine/Source/Editor/WorldBookmark/Private/WorldBookmark/WorldBookmarkModule.h
+++ b/Engine/Source/Editor/WorldBookmark/Private/WorldBookmark/WorldBookmarkModule.h
@@ -28,6 +28,26 @@ public:
 	 */
 	virtual void ShutdownModule();
 
+    /** 
+     * Capture a bookmark as a string, can be restored with RestoreFromString().
+	 */
+	static FString CaptureToString();
+
+    /** 
+     * Capture a bookmark to the clipboard, can be restored with RestoreFromClipboard().
+	 */
+	static void CaptureToClipboard();
+
+    /** 
+     * Restore a bookmark from the provided string.
+	 */
+	static bool RestoreFromString(const FString& InBookmarkAsString);
+
+    /** 
+     * Restore a bookmark from the clipboard.
+	 */
+	static bool RestoreFromClipboard();
+
 private:
 	void RegisterWorldBookmarkBrowserTab(TSharedPtr<FTabManager> InTabManager);
 	TSharedRef<SDockTab> SpawnWorldBookmarkBrowserTab(const FSpawnTabArgs& Args);
diff --git a/Engine/Source/Editor/WorldBookmark/WorldBookmark.Build.cs b/Engine/Source/Editor/WorldBookmark/WorldBookmark.Build.cs
--- a/Engine/Source/Editor/WorldBookmark/WorldBookmark.Build.cs
+++ b/Engine/Source/Editor/WorldBookmark/WorldBookmark.Build.cs
@@ -8,6 +8,7 @@ public class WorldBookmark : ModuleRules
 	{
 		PrivateDependencyModuleNames.AddRange(
 			new string[] {
+				"ApplicationCore",
 				"AssetDefinition",
 				"AssetRegistry",
 				"AssetTools",
@@ -20,6 +21,8 @@ public class WorldBookmark : ModuleRules
 				"EditorConfig",
 				"Engine",
 				"InputCore",
+				"Json",
+				"JsonUtilities",
 				"LevelEditor",
 				"Projects",
 				"PropertyEditor",
