World Bookmark - Fix issue where "Play from Here" was only transitioning to PIE if the user selected a bookmark for a different world
* When loading a bookmark is going to cause a map change, show a suppressable dialog warning the user. Play from here/move camera here were already showing a dialog, but the double click or "Load" actions weren't

diff --git a/Engine/Source/Editor/WorldBookmark/Private/WorldBookmark/Browser/SWorldBookmarkBrowser.cpp b/Engine/Source/Editor/WorldBookmark/Private/WorldBookmark/Browser/SWorldBookmarkBrowser.cpp
--- a/Engine/Source/Editor/WorldBookmark/Private/WorldBookmark/Browser/SWorldBookmarkBrowser.cpp
+++ b/Engine/Source/Editor/WorldBookmark/Private/WorldBookmark/Browser/SWorldBookmarkBrowser.cpp
@@ -21,6 +21,7 @@
 #include "AssetRegistry/AssetRegistryHelpers.h"
 #include "AssetToolsModule.h"
 #include "ContentBrowserModule.h"
+#include "Dialogs/Dialogs.h"
 #include "Editor.h"
 #include "Framework/Application/SlateApplication.h"
 #include "Framework/Commands/GenericCommands.h"
@@ -822,6 +823,11 @@ void SWorldBookmarkBrowser::OnMouseButtonDoubleClicked(FTreeItemPtr InTreeItem)
 {
 	if (UWorldBookmark* WorldBookmark = GetSelectedBookmark())
 	{
+		if (!AskForWorldChangeConfirmation(WorldBookmark))
+		{
+			return;
+		}
+
 		WorldBookmark->Load();
 
 		// Loading an item will update it's last load time - refresh the UI
@@ -1083,6 +1089,11 @@ void SWorldBookmarkBrowser::LoadSelectedBookmark()
 {
 	if (UWorldBookmark* SelectedBookmark = GetSelectedBookmark())
 	{
+		if (!AskForWorldChangeConfirmation(SelectedBookmark))
+		{
+			return;
+		}
+
 		SelectedBookmark->Load();
 	}
 }
@@ -1262,7 +1273,7 @@ void SWorldBookmarkBrowser::DeleteSelectedItem()
 	}
 }
 
-static bool AskForWorldChangeConfirmation(UWorldBookmark* WorldBookmark)
+bool SWorldBookmarkBrowser::AskForWorldChangeConfirmation(UWorldBookmark* WorldBookmark) const
 {
 	if (const UWorldEditorState* WorldEditorState = WorldBookmark->GetEditorState<UWorldEditorState>())
 	{
@@ -1271,12 +1282,18 @@ static bool AskForWorldChangeConfirmation(UWorldBookmark* WorldBookmark)
 		{
 			const FText MessageBoxTitle = LOCTEXT("WorldBookmark_DifferentWorldDlg_Title", "Open Bookmark World?");
 			const FText MessageBoxText = LOCTEXT("WorldBookmark_DifferentWorldDlg_Text", "This bookmark is for a different world, are you sure you want to open it?");
-			const EAppReturnType::Type Response = FMessageDialog::Open(EAppMsgType::YesNo, MessageBoxText, MessageBoxTitle);
-			return Response == EAppReturnType::Yes;
+			
+			FSuppressableWarningDialog::FSetupInfo Info(MessageBoxText, MessageBoxTitle, "WorldBookmark_AskForWorldChangeConfirmation");
+			Info.ConfirmText = FCoreTexts::Get().Yes;
+			Info.CancelText = FCoreTexts::Get().No;
+
+			FSuppressableWarningDialog AddLevelWarning(Info);
+			const FSuppressableWarningDialog::EResult Response = AddLevelWarning.ShowModal();
+			return Response != FSuppressableWarningDialog::EResult::Cancel;
 		}
 	}
 
-	return false;
+	return true;
 }
 
 bool SWorldBookmarkBrowser::CanGoToSelectedBookmarkLocation() const
diff --git a/Engine/Source/Editor/WorldBookmark/Private/WorldBookmark/Browser/SWorldBookmarkBrowser.h b/Engine/Source/Editor/WorldBookmark/Private/WorldBookmark/Browser/SWorldBookmarkBrowser.h
--- a/Engine/Source/Editor/WorldBookmark/Private/WorldBookmark/Browser/SWorldBookmarkBrowser.h
+++ b/Engine/Source/Editor/WorldBookmark/Private/WorldBookmark/Browser/SWorldBookmarkBrowser.h
@@ -106,6 +106,8 @@ private:
 	void MoveSelectedBookmarkToNewFolder();
 	void CreateBookmarkInSelectedFolder();
 
+	bool AskForWorldChangeConfirmation(UWorldBookmark* WorldBookmark) const;
+
 private:
 	FDelegateHandle OnAssetAddedHandle;
 	FDelegateHandle OnAssetRemovedHandle;
