[World Bookmark][UX] Favorite Icon is now the color of the bookmark category.

#tests world bookmark editor

diff --git a/Engine/Source/Editor/WorldBookmark/Private/WorldBookmark/Browser/BookmarkTableRow.cpp b/Engine/Source/Editor/WorldBookmark/Private/WorldBookmark/Browser/BookmarkTableRow.cpp
--- a/Engine/Source/Editor/WorldBookmark/Private/WorldBookmark/Browser/BookmarkTableRow.cpp
+++ b/Engine/Source/Editor/WorldBookmark/Private/WorldBookmark/Browser/BookmarkTableRow.cpp
@@ -37,7 +37,15 @@ TSharedRef<SWidget> FWorldBookmarkTableRow::GenerateWidgetForColumn(const FName&
 		{
 			if (WorldBookmark->GetIsUserFavorite())
 			{
-				return CreateIconWidget(GetFavoriteWorldBookmarkIcon(true));
+				return SNew(SBox)
+					.HAlign(HAlign_Center)
+					.VAlign(VAlign_Center)
+					[
+						SNew(SImage)
+							.DesiredSizeOverride(FVector2D(16, 16))
+							.ColorAndOpacity(this, &FWorldBookmarkTableRow::GetIconColor)
+							.Image(GetFavoriteWorldBookmarkIcon(true).GetIcon())
+					];
 			}
 		}
 		else if (ColumnName == Columns::RecentlyUsed.Id)
