[WorldBookmark][UX]
- Apply the hidden columns from the saved config.
- Preserve the user hidden columns when recreating them.

#tests World Bookmark Outliner
[RN]

diff --git a/Engine/Source/Editor/WorldBookmark/Private/WorldBookmark/Browser/SWorldBookmarkBrowser.cpp b/Engine/Source/Editor/WorldBookmark/Private/WorldBookmark/Browser/SWorldBookmarkBrowser.cpp
--- a/Engine/Source/Editor/WorldBookmark/Private/WorldBookmark/Browser/SWorldBookmarkBrowser.cpp
+++ b/Engine/Source/Editor/WorldBookmark/Private/WorldBookmark/Browser/SWorldBookmarkBrowser.cpp
@@ -120,7 +120,6 @@ void SWorldBookmarkBrowser::Construct(const FArguments& InArgs)
 
 	HeaderRow = SNew(SHeaderRow)
 		.CanSelectGeneratedColumn(true)
-		.HiddenColumnsList(UWorldBookmarkBrowserSettings::GetHiddenColumns())
 		.OnHiddenColumnsListChanged_Lambda([this]()
 		{
 			UWorldBookmarkBrowserSettings::SetHiddenColumns(BookmarksView->GetHeaderRow()->GetHiddenColumnIds());
@@ -143,7 +142,7 @@ void SWorldBookmarkBrowser::Construct(const FArguments& InArgs)
 		);
 
 	TSharedRef<SSearchBox> SearchBoxRef = SNew(SSearchBox)
-		.OnTextChanged(this, &SWorldBookmarkBrowser::OnSearchBoxTextChanged);		
+		.OnTextChanged(this, &SWorldBookmarkBrowser::OnSearchBoxTextChanged);
 
 	ChildSlot
 	[
@@ -160,7 +159,7 @@ void SWorldBookmarkBrowser::Construct(const FArguments& InArgs)
 			.AutoHeight()
 			.Padding(8.0f, 8.0f, 8.0f, 4.0f)
 			.VAlign(VAlign_Top)
-			.HAlign(HAlign_Fill)			
+			.HAlign(HAlign_Fill)
 			[
 				SNew(SHorizontalBox)
 				
@@ -355,6 +354,15 @@ void SWorldBookmarkBrowser::RecreateColumns()
 		HeaderRow->AddColumn(LabelColumn);
 		HeaderRow->AddColumn(WorldColumn);
 	}
+
+	// Reapply the hidden columns.
+	const TArray<FName>& HiddenColumnsIds = UWorldBookmarkBrowserSettings::GetHiddenColumns();
+	for (const SHeaderRow::FColumn& Column : HeaderRow->GetColumns())
+	{
+		SHeaderRow::FColumn& Col = const_cast<SHeaderRow::FColumn&>(Column);
+		Col.bIsVisible = !HiddenColumnsIds.Contains(Column.ColumnId);
+	}
+	HeaderRow->RefreshColumns();
 }
 
 TSharedRef<SWidget> SWorldBookmarkBrowser::GetSettingsMenuContent()
