[World Bookmarks] Add tooltip for Recently Used icons. Minor capitalization fix to World Bookmark Name column label.

diff --git a/Engine/Source/Editor/WorldBookmark/Private/WorldBookmark/Browser/BookmarkTableRow.cpp b/Engine/Source/Editor/WorldBookmark/Private/WorldBookmark/Browser/BookmarkTableRow.cpp
--- a/Engine/Source/Editor/WorldBookmark/Private/WorldBookmark/Browser/BookmarkTableRow.cpp
+++ b/Engine/Source/Editor/WorldBookmark/Private/WorldBookmark/Browser/BookmarkTableRow.cpp
@@ -53,7 +53,7 @@ TSharedRef<SWidget> FWorldBookmarkTableRow::GenerateWidgetForColumn(const FName&
 			FTimespan LastLoaded = FDateTime::UtcNow() - WorldBookmark->GetUserLastLoadedTimeStampUTC();
 			if (LastLoaded < FTimespan::FromDays(14))
 			{
-				return CreateIconWidget(GetRecentlyUsedWorldBookmarkIcon());
+				return CreateIconWidget(GetRecentlyUsedWorldBookmarkIcon(), GetLastAccessedText());
 			}
 		}
 		else if (ColumnName == Columns::Label.Id)
diff --git a/Engine/Source/Editor/WorldBookmark/Private/WorldBookmark/Browser/Columns.cpp b/Engine/Source/Editor/WorldBookmark/Private/WorldBookmark/Browser/Columns.cpp
--- a/Engine/Source/Editor/WorldBookmark/Private/WorldBookmark/Browser/Columns.cpp
+++ b/Engine/Source/Editor/WorldBookmark/Private/WorldBookmark/Browser/Columns.cpp
@@ -29,7 +29,7 @@ namespace Columns
 	{
 		TEXT("Label"),
 		LOCTEXT("Column_Label_Text", "Name"),
-		LOCTEXT("Column_Label_Tooltip", "World Bookmark name")
+		LOCTEXT("Column_Label_Tooltip", "World Bookmark Name")
 	};
 
 	const FColumnDefinition World =
diff --git a/Engine/Source/Editor/WorldBookmark/Private/WorldBookmark/Browser/Icons.cpp b/Engine/Source/Editor/WorldBookmark/Private/WorldBookmark/Browser/Icons.cpp
--- a/Engine/Source/Editor/WorldBookmark/Private/WorldBookmark/Browser/Icons.cpp
+++ b/Engine/Source/Editor/WorldBookmark/Private/WorldBookmark/Browser/Icons.cpp
@@ -34,7 +34,7 @@ FSlateIcon GetWorldBookmarkIcon()
 	return FSlateIcon(FWorldBookmarkStyle::Get().GetStyleSetName(), WorldBookmarkIconName);
 }
 
-TSharedRef<SWidget> CreateIconWidget(const FSlateIcon& SlateIcon)
+TSharedRef<SWidget> CreateIconWidget(const FSlateIcon& SlateIcon, const FText& ToolTip)
 {
 	return SNew(SBox)
 		.HAlign(HAlign_Center)
@@ -44,6 +44,7 @@ TSharedRef<SWidget> CreateIconWidget(const FSlateIcon& SlateIcon)
 				.DesiredSizeOverride(FVector2D(16, 16))
 				.ColorAndOpacity(FSlateColor::UseForeground())
 				.Image(SlateIcon.GetIcon())
+				.ToolTipText(ToolTip)
 		];
 }
 
diff --git a/Engine/Source/Editor/WorldBookmark/Private/WorldBookmark/Browser/Icons.h b/Engine/Source/Editor/WorldBookmark/Private/WorldBookmark/Browser/Icons.h
--- a/Engine/Source/Editor/WorldBookmark/Private/WorldBookmark/Browser/Icons.h
+++ b/Engine/Source/Editor/WorldBookmark/Private/WorldBookmark/Browser/Icons.h
@@ -12,5 +12,5 @@ namespace UE::WorldBookmark::Browser
 	FSlateIcon GetRecentlyUsedWorldBookmarkIcon();
 	FSlateIcon GetFolderIcon(bool bIsExpanded);
 	FSlateIcon GetWorldBookmarkIcon();
-	TSharedRef<SWidget> CreateIconWidget(const FSlateIcon& SlateIcon);
+	TSharedRef<SWidget> CreateIconWidget(const FSlateIcon& SlateIcon, const FText& ToolTip = FText());
 }
\ No newline at end of file
