World Bookmarks Category Window Fixes
- Refresh UI when cancelling the creation of a new category. This fixes the selection for the next time new category is selected.
- Assigning a random color to a new category in CreateNewCategory. FWorldBookmarkCategory default color is now a random seded color based on name.

diff --git a/Engine/Source/Editor/WorldBookmark/Private/WorldBookmark/WorldBookmark.cpp b/Engine/Source/Editor/WorldBookmark/Private/WorldBookmark/WorldBookmark.cpp
--- a/Engine/Source/Editor/WorldBookmark/Private/WorldBookmark/WorldBookmark.cpp
+++ b/Engine/Source/Editor/WorldBookmark/Private/WorldBookmark/WorldBookmark.cpp
@@ -298,7 +298,7 @@ const FWorldBookmarkCategory& UWorldBookmark::GetBookmarkCategory() const
 const FWorldBookmarkCategory FWorldBookmarkCategory::None = FWorldBookmarkCategory(NAME_None, FColor::Black);
 
 FWorldBookmarkCategory::FWorldBookmarkCategory()
-	: FWorldBookmarkCategory(TEXT("New Category"), FColor::MakeRandomColor())
+	: FWorldBookmarkCategory(TEXT("New Category"), FColor::MakeRandomSeededColor(GetTypeHash(FString(TEXT("New Category")))))
 {
 }
 
diff --git a/Engine/Source/Editor/WorldBookmark/Private/WorldBookmark/WorldBookmarkDetailsCustomization.cpp b/Engine/Source/Editor/WorldBookmark/Private/WorldBookmark/WorldBookmarkDetailsCustomization.cpp
--- a/Engine/Source/Editor/WorldBookmark/Private/WorldBookmark/WorldBookmarkDetailsCustomization.cpp
+++ b/Engine/Source/Editor/WorldBookmark/Private/WorldBookmark/WorldBookmarkDetailsCustomization.cpp
@@ -232,6 +232,13 @@ void FWorldBookmarkDetailsCustomization::OnCategoryChanged(TSharedPtr<FWorldBook
 			{
 				UWorldBookmarkEditorSettings::AddCategory(CategoryToAssign);
 			}
+			else
+			{
+				if (CachedDetailBuilder)
+				{
+					CachedDetailBuilder->ForceRefreshDetails();
+				}
+			}
 			bShouldAssign = bCreatedNewCategory;
 		}
 
@@ -266,6 +273,7 @@ bool FWorldBookmarkDetailsCustomization::CreateNewCategory(FWorldBookmarkCategor
 	ViewArgs.bShowObjectLabel = false;
 
 	FWorldBookmarkCategory NewCategoryTemplate;
+	NewCategoryTemplate.Color = FColor::MakeRandomColor();
 	TSharedPtr<TStructOnScope<FWorldBookmarkCategory>> NewCategory;
 	NewCategory = MakeShared<TStructOnScope<FWorldBookmarkCategory>>();
 	NewCategory->InitializeAs<FWorldBookmarkCategory>(NewCategoryTemplate);
