[Transforms] Defer loading of editor specific modules to editor init phase for world bookmarks

diff --git a/Engine/Source/Editor/WorldBookmark/Private/WorldBookmark/WorldBookmarkModule.cpp b/Engine/Source/Editor/WorldBookmark/Private/WorldBookmark/WorldBookmarkModule.cpp
--- a/Engine/Source/Editor/WorldBookmark/Private/WorldBookmark/WorldBookmarkModule.cpp
+++ b/Engine/Source/Editor/WorldBookmark/Private/WorldBookmark/WorldBookmarkModule.cpp
@@ -10,6 +10,7 @@
 
 #include "AssetRegistry/IAssetRegistry.h"
 #include "AssetRegistry/AssetData.h"
+#include "Editor.h"
 #include "EditorState/EditorStateSubsystem.h"
 #include "EditorState/WorldEditorState.h"
 #include "Framework/Notifications/NotificationManager.h"
@@ -95,7 +96,7 @@ static FAutoConsoleCommand WorldBookmarkRestoreFromClipboardCommand(
 
 void FWorldBookmarkModule::StartupModule()
 {
-	if (GIsEditor && !IsRunningCommandlet())
+	FEditorDelegates::OnEditorBoot.AddLambda([this](double)
 	{
 		FWorldBookmarkCommands::Register();
 
@@ -121,7 +122,7 @@ void FWorldBookmarkModule::StartupModule()
 
 		// Validate level default bookmark changes
 		OnDefaultBookmarkChangedHandle = AWorldSettings::OnDefaultBookmarkChanged.AddRaw(this, &FWorldBookmarkModule::OnDefaultBookmarkChanged);
-	}
+	});
 }
 
 void FWorldBookmarkModule::ShutdownModule()
