World Bookmarks - Fixed mismatched delegate registration
* Was registering to OnAssetsAdded() and unregistering to OnAssetAdded()

#Synced-CL 50899516

diff --git a/Engine/Source/Editor/WorldBookmark/Private/WorldBookmark/Browser/SWorldBookmarkBrowser.cpp b/Engine/Source/Editor/WorldBookmark/Private/WorldBookmark/Browser/SWorldBookmarkBrowser.cpp
--- a/Engine/Source/Editor/WorldBookmark/Private/WorldBookmark/Browser/SWorldBookmarkBrowser.cpp
+++ b/Engine/Source/Editor/WorldBookmark/Private/WorldBookmark/Browser/SWorldBookmarkBrowser.cpp
@@ -92,7 +92,7 @@ SWorldBookmarkBrowser::SWorldBookmarkBrowser()
 SWorldBookmarkBrowser::~SWorldBookmarkBrowser()
 {
 	IAssetRegistry& AssetRegistry = IAssetRegistry::GetChecked();
-	AssetRegistry.OnAssetAdded().Remove(OnAssetAddedHandle);
+	AssetRegistry.OnAssetsAdded().Remove(OnAssetAddedHandle);
 	AssetRegistry.OnAssetRemoved().Remove(OnAssetRemovedHandle);
 	AssetRegistry.OnAssetRenamed().Remove(OnAssetRenamedHandle);
 	AssetRegistry.OnAssetUpdated().Remove(OnAssetUpdatedHandle);
