[ContentBrowser] New Add Menu World Menu

#Synced-CL 52439714

diff --git a/Engine/Plugins/Editor/EngineAssetDefinitions/Source/Private/AssetDefinition_DataLayer.h b/Engine/Plugins/Editor/EngineAssetDefinitions/Source/Private/AssetDefinition_DataLayer.h
--- a/Engine/Plugins/Editor/EngineAssetDefinitions/Source/Private/AssetDefinition_DataLayer.h
+++ b/Engine/Plugins/Editor/EngineAssetDefinitions/Source/Private/AssetDefinition_DataLayer.h
@@ -21,7 +21,7 @@ public:
 	virtual FAssetSupportResponse CanLocalize(const FAssetData& InAsset) const override { return FAssetSupportResponse::NotSupported(); }
 	virtual TConstArrayView<FAssetCategoryPath> GetAssetCategories() const override
 	{
-		static const auto Categories = { EAssetCategoryPaths::World };
+		static const FAssetCategoryPath Categories[] = { FAssetCategoryPath(EAssetCategoryPaths::World, NSLOCTEXT("AssetDefinition", "DataLayer_SubMenu", "World"), ECategoryMenuType::Section) };
 		return Categories;
 	}
 	// UAssetDefinition End
diff --git a/Engine/Plugins/Editor/EngineAssetDefinitions/Source/Private/AssetDefinition_HLODLayer.h b/Engine/Plugins/Editor/EngineAssetDefinitions/Source/Private/AssetDefinition_HLODLayer.h
--- a/Engine/Plugins/Editor/EngineAssetDefinitions/Source/Private/AssetDefinition_HLODLayer.h
+++ b/Engine/Plugins/Editor/EngineAssetDefinitions/Source/Private/AssetDefinition_HLODLayer.h
@@ -21,7 +21,7 @@ public:
 	virtual FAssetSupportResponse CanLocalize(const FAssetData& InAsset) const override { return FAssetSupportResponse::NotSupported(); }
 	virtual TConstArrayView<FAssetCategoryPath> GetAssetCategories() const override
 	{
-		static const auto Categories = { EAssetCategoryPaths::World };
+		static const FAssetCategoryPath Categories[] = { FAssetCategoryPath(EAssetCategoryPaths::World, NSLOCTEXT("AssetDefinition", "HLODLayer_SubMenu", "World"), ECategoryMenuType::Section) };
 		return Categories;
 	}
 	// UAssetDefinition End
diff --git a/Engine/Plugins/Editor/EngineAssetDefinitions/Source/Private/AssetDefinition_World.cpp b/Engine/Plugins/Editor/EngineAssetDefinitions/Source/Private/AssetDefinition_World.cpp
--- a/Engine/Plugins/Editor/EngineAssetDefinitions/Source/Private/AssetDefinition_World.cpp
+++ b/Engine/Plugins/Editor/EngineAssetDefinitions/Source/Private/AssetDefinition_World.cpp
@@ -19,7 +19,11 @@
 
 TConstArrayView<FAssetCategoryPath> UAssetDefinition_World::GetAssetCategories() const
 {
-	static const auto Categories = { EAssetCategoryPaths::Basic };
+	static const FAssetCategoryPath Categories[] = 
+		{
+			FAssetCategoryPath(EAssetCategoryPaths::Basic),
+			FAssetCategoryPath(EAssetCategoryPaths::World, LOCTEXT("AssetDefinition_WorldSubMenu", "World"), ECategoryMenuType::Section)
+		};
 	return Categories;
 }
 
diff --git a/Engine/Source/Editor/WorldBookmark/Private/WorldBookmark/AssetDefinition_WorldBookmark.h b/Engine/Source/Editor/WorldBookmark/Private/WorldBookmark/AssetDefinition_WorldBookmark.h
--- a/Engine/Source/Editor/WorldBookmark/Private/WorldBookmark/AssetDefinition_WorldBookmark.h
+++ b/Engine/Source/Editor/WorldBookmark/Private/WorldBookmark/AssetDefinition_WorldBookmark.h
@@ -20,9 +20,9 @@ public:
 	virtual TSoftClassPtr<UObject> GetAssetClass() const override { return UWorldBookmark::StaticClass(); }
 	virtual FAssetSupportResponse CanLocalize(const FAssetData& InAsset) const override { return FAssetSupportResponse::NotSupported(); }	
 
-	virtual TConstArrayView<FAssetCategoryPath> GetAssetCategories() const
+	virtual TConstArrayView<FAssetCategoryPath> GetAssetCategories() const override
 	{
-		static const auto Categories = { EAssetCategoryPaths::World };
+		static const FAssetCategoryPath Categories[] = { FAssetCategoryPath(EAssetCategoryPaths::World, NSLOCTEXT("AssetDefinition", "WorldBookmark_SubMenu", "World"), ECategoryMenuType::Section) };
 		return Categories;
 	}
 
